UE4: Even more logs to help track down the cook crash.

diff --git a/Engine/Source/Developer/TargetPlatform/Private/TargetPlatformManagerModule.cpp b/Engine/Source/Developer/TargetPlatform/Private/TargetPlatformManagerModule.cpp
--- a/Engine/Source/Developer/TargetPlatform/Private/TargetPlatformManagerModule.cpp
+++ b/Engine/Source/Developer/TargetPlatform/Private/TargetPlatformManagerModule.cpp
@@ -24,11 +24,17 @@ public:
 	FTargetPlatformManagerModule()
 		: bRestrictFormatsToRuntimeOnly(false), bForceCacheUpdate(true)
 	{
+		UE_LOG(LogInit, Log, TEXT("FTargetPlatformManagerModule constructor"));
 		GetTargetPlatforms();
+		UE_LOG(LogInit, Log, TEXT("GetTargetPlatforms complete"));
 		GetActiveTargetPlatforms();
+		UE_LOG(LogInit, Log, TEXT("GetActiveTargetPlatforms complete"));
 		GetAudioFormats();
+		UE_LOG(LogInit, Log, TEXT("GetAudioFormats complete"));
 		GetTextureFormats();
+		UE_LOG(LogInit, Log, TEXT("GetTextureFormats complete"));
 		GetShaderFormats();
+		UE_LOG(LogInit, Log, TEXT("GetShaderFormats complete"));
 
 		bForceCacheUpdate = false;
 
@@ -177,13 +183,13 @@ public:
 
 			if (!Results.Num())
 			{
-				UE_LOG(LogTargetPlatformManager, Display, TEXT("Not building assets for any platform."));
+				UE_LOG(LogTargetPlatformManager, Log, TEXT("Not building assets for any platform."));
 			}
 			else
 			{
 				for (int32 Index = 0; Index < Results.Num(); Index++)
 				{
-					UE_LOG(LogTargetPlatformManager, Display, TEXT("Building Assets For %s"), *Results[Index]->PlatformName());
+					UE_LOG(LogTargetPlatformManager, Log, TEXT("Building Assets For %s"), *Results[Index]->PlatformName());
 				}
 			}
 		}
diff --git a/Engine/Source/Runtime/Core/Private/Modules/ModuleManager.cpp b/Engine/Source/Runtime/Core/Private/Modules/ModuleManager.cpp
--- a/Engine/Source/Runtime/Core/Private/Modules/ModuleManager.cpp
+++ b/Engine/Source/Runtime/Core/Private/Modules/ModuleManager.cpp
@@ -344,7 +344,7 @@ TSharedPtr<IModuleInterface> FModuleManager::LoadModuleWithFailureReason( const
 
 			// Try to dynamically load the DLL
 
-			UE_LOG(LogModuleManager, Verbose, TEXT( "ModuleManager: Load Module '%s' DLL '%s'" ), *InModuleName.ToString(), *ModuleInfo->Filename);
+			UE_LOG(LogModuleManager, Log, TEXT( "ModuleManager: Load Module '%s' DLL '%s'" ), *InModuleName.ToString(), *ModuleInfo->Filename);
 
 			// Determine which file to load for this module.
 			const FString ModuleFileToLoad = ModuleInfo->Filename;
