TargetPlatform: Documentation and formatting clean up pass

diff --git a/Engine/Source/Developer/TargetPlatform/Public/Interfaces/IAudioFormat.h b/Engine/Source/Developer/TargetPlatform/Public/Interfaces/IAudioFormat.h
--- a/Engine/Source/Developer/TargetPlatform/Public/Interfaces/IAudioFormat.h
+++ b/Engine/Source/Developer/TargetPlatform/Public/Interfaces/IAudioFormat.h
@@ -13,34 +13,22 @@
  */
 struct FSoundQualityInfo
 {
-	/**
-	 * Holds the quality value ranging from 1 [poor] to 100 [very good].
-	 */
+	/** Holds the quality value ranging from 1 [poor] to 100 [very good]. */
 	int32 Quality;
 
-	/**
-	 * Holds the number of distinct audio channels.
-	 */
+	/** Holds the number of distinct audio channels. */
 	uint32 NumChannels;
 
-	/**
-	 * Holds the number of PCM samples per second.
-	 */
+	/** Holds the number of PCM samples per second. */
 	uint32 SampleRate;
 
-	/**
-	 * Holds the size of sample data in bytes.
-	 */
+	/** Holds the size of sample data in bytes. */
 	uint32 SampleDataSize;
 
-	/**
-	 * Holds the length of the sound in seconds.
-	 */
+	/** Holds the length of the sound in seconds. */
 	float Duration;
 
-	/**
-	 * Holds a string for debugging purposes.
-	 */
+	/** Holds a string for debugging purposes. */
 	FString DebugName;
 };
 
@@ -67,39 +55,36 @@ public:
 	/**
 	 * Cooks the source data for the platform and stores the cooked data internally.
 	 *
-	 * @param Format - The desired format.
-	 * @param SrcBuffer - The source buffer.
-	 * @param QualityInfo - All the information the compressor needs to compress the audio.
-	 * @param OutBuffer - Will hold the resulting compressed audio.
-	 *
-	 * @return	true if succeeded, false otherwise.
+	 * @param Format The desired format.
+	 * @param SrcBuffer The source buffer.
+	 * @param QualityInfo All the information the compressor needs to compress the audio.
+	 * @param OutBuffer Will hold the resulting compressed audio.
+	 * @return true on success, false otherwise.
 	 */
 	virtual bool Cook( FName Format, const TArray<uint8>& SrcBuffer, FSoundQualityInfo& QualityInfo, TArray<uint8>& OutBuffer ) const = 0;
 
 	/**
 	 * Cooks up to 8 mono files into a multi-stream file (e.g. 5.1). The front left channel is required, the rest are optional.
 	 *
-	 * @param	Format			desired format
-	 * @param	SrcBuffers		Source buffers
-	 * @param	QualityInfo		All the information the compressor needs to compress the audio
-	 * @param	OutBuffer		resulting compressed audio
-	 *
-	 * @return	true if succeeded, false otherwise
+	 * @param Format The desired format.
+	 * @param SrcBuffers The source buffers.
+	 * @param QualityInfo All the information the compressor needs to compress the audio
+	 * @param OutBuffer Will contain the resulting compressed audio.
+	 * @return true if succeeded, false otherwise
 	 */
 	virtual bool CookSurround( FName Format, const TArray<TArray<uint8> >& SrcBuffers, FSoundQualityInfo& QualityInfo, TArray<uint8>& OutBuffer ) const = 0;
 
 	/**
 	 * Gets the list of supported formats.
 	 *
-	 * @param OutFormats - Will hold the list of formats.
+	 * @param OutFormats Will contain the list of formats.
 	 */
 	virtual void GetSupportedFormats( TArray<FName>& OutFormats ) const = 0;
 
 	/**
 	 * Gets the current version of the specified audio format.
 	 *
-	 * @param Format - The format to get the version for.
-	 *
+	 * @param Format The format to get the version for.
 	 * @return Version number.
 	 */
 	virtual uint16 GetVersion( FName Format ) const = 0;
@@ -107,11 +92,10 @@ public:
 	/** 
 	 * Re-compresses raw PCM to the the platform dependent format, and then back to PCM. Used for quality previewing.
 	 *
-	 * @param Format - The desired format.
-	 * @param SrcBuffer - Uncompressed PCM data.
-	 * @param QualityInfo - All the information the compressor needs to compress the audio.
-	 * @param OutBuffer - Uncompressed PCM data after being compressed.
-	 *
+	 * @param Format The desired format.
+	 * @param SrcBuffer Uncompressed PCM data.
+	 * @param QualityInfo All the information the compressor needs to compress the audio.
+	 * @param OutBuffer Uncompressed PCM data after being compressed.
 	 * @return The size of the compressed audio, or 0 on failure.
 	 */
 	virtual int32 Recompress( FName Format, const TArray<uint8>& SrcBuffer, FSoundQualityInfo& QualityInfo, TArray<uint8>& OutBuffer ) const = 0;
diff --git a/Engine/Source/Developer/TargetPlatform/Public/Interfaces/IAudioFormatModule.h b/Engine/Source/Developer/TargetPlatform/Public/Interfaces/IAudioFormatModule.h
--- a/Engine/Source/Developer/TargetPlatform/Public/Interfaces/IAudioFormatModule.h
+++ b/Engine/Source/Developer/TargetPlatform/Public/Interfaces/IAudioFormatModule.h
@@ -18,10 +18,12 @@ public:
 	/**
 	 * Gets the audio format.
 	 */
-	virtual IAudioFormat* GetAudioFormat() = 0;
+	virtual IAudioFormat* GetAudioFormat( ) = 0;
 
+public:
 
-protected:
-
-	IAudioFormatModule() { }
+	/**
+	 * Virtual destructor.
+	 */
+	~IAudioFormatModule( ) { }
 };
diff --git a/Engine/Source/Developer/TargetPlatform/Public/Interfaces/IPhysXFormat.h b/Engine/Source/Developer/TargetPlatform/Public/Interfaces/IPhysXFormat.h
--- a/Engine/Source/Developer/TargetPlatform/Public/Interfaces/IPhysXFormat.h
+++ b/Engine/Source/Developer/TargetPlatform/Public/Interfaces/IPhysXFormat.h
@@ -29,50 +29,46 @@ public:
 	/**
 	 * Cooks the source convex data for the platform and stores the cooked data internally.
 	 *
-	 * @param Format - The desired format
-	 * @param SrcBuffer - The source buffer
-	 * @param OutBuffer - The resulting cooked data
-	 *
-	 * @return	true if succeeded, false otherwise.
+	 * @param Format The desired format
+	 * @param SrcBuffer The source buffer
+	 * @param OutBuffer The resulting cooked data
+	 * @return true on success, false otherwise.
 	 */
 	virtual bool CookConvex( FName Format, const TArray<FVector>& SrcBuffer, TArray<uint8>& OutBuffer ) const = 0;
 
 	/**
 	 * Cooks the source Tri-Mesh data for the platform and stores the cooked data internally.
 	 *
-	 * @param Format - The desired format.
-	 * @param SrcBuffer - The source buffer.
-	 * @param OutBuffer - The resulting cooked data.
-	 *
-	 * @return	true if succeeded, false otherwise.
+	 * @param Format The desired format.
+	 * @param SrcBuffer The source buffer.
+	 * @param OutBuffer The resulting cooked data.
+	 * @return true on success, false otherwise.
 	 */
 	virtual bool CookTriMesh( FName Format, const TArray<FVector>& SrcVertices, const TArray<struct FTriIndices>& SrcIndices, const TArray<uint16>& SrcMaterialIndices, const bool FlipNormals, TArray<uint8>& OutBuffer ) const = 0;
 		
 	/**
 	 * Cooks the source height field data for the platform and stores the cooked data internally.
 	 *
-	 * @param Format - The desired format
-	 * @param HFSize - Size of height field [NumColumns, NumRows]
-	 * @param Thickness - Sets how thick the height field surface is
-	 * @param SrcBuffer - The source buffer
-	 * @param OutBuffer - The resulting cooked data
-	 *
-	 * @return	true if succeeded, false otherwise.
+	 * @param Format The desired format
+	 * @param HFSize Size of height field [NumColumns, NumRows]
+	 * @param Thickness Sets how thick the height field surface is
+	 * @param SrcBuffer The source buffer
+	 * @param OutBuffer The resulting cooked data
+	 * @return true on success, false otherwise.
 	 */
 	virtual bool CookHeightField( FName Format, FIntPoint HFSize, float Thickness, const void* Samples, uint32 SamplesStride, TArray<uint8>& OutBuffer ) const = 0;
 
 	/**
 	 * Gets the list of supported formats.
 	 *
-	 * @param OutFormats - Will hold the list of formats.
+	 * @param OutFormats Will hold the list of formats.
 	 */
 	virtual void GetSupportedFormats( TArray<FName>& OutFormats ) const = 0;
 
 	/**
 	 * Gets the current version of the specified PhysX format.
 	 *
-	 * @param Format - The format to get the version for.
-	 *
+	 * @param Format The format to get the version for.
 	 * @return Version number.
 	 */
 	virtual uint16 GetVersion( FName Format ) const = 0;
diff --git a/Engine/Source/Developer/TargetPlatform/Public/Interfaces/IPhysXFormatModule.h b/Engine/Source/Developer/TargetPlatform/Public/Interfaces/IPhysXFormatModule.h
--- a/Engine/Source/Developer/TargetPlatform/Public/Interfaces/IPhysXFormatModule.h
+++ b/Engine/Source/Developer/TargetPlatform/Public/Interfaces/IPhysXFormatModule.h
@@ -18,10 +18,12 @@ public:
 	/**
 	 * Gets the PhysX format.
 	 */
-	virtual IPhysXFormat* GetPhysXFormat() = 0;
+	virtual IPhysXFormat* GetPhysXFormat( ) = 0;
 
+public:
 
-protected:
-
-	IPhysXFormatModule() { }
+	/**
+	 * Virtual destructor.
+	 */
+	~IPhysXFormatModule( ) { }
 };
diff --git a/Engine/Source/Developer/TargetPlatform/Public/Interfaces/IShaderFormat.h b/Engine/Source/Developer/TargetPlatform/Public/Interfaces/IShaderFormat.h
--- a/Engine/Source/Developer/TargetPlatform/Public/Interfaces/IShaderFormat.h
+++ b/Engine/Source/Developer/TargetPlatform/Public/Interfaces/IShaderFormat.h
@@ -6,9 +6,11 @@
 
 #pragma once
 
+
 /** Wildcard string used to search for shader format modules. */
 #define SHADERFORMAT_MODULE_WILDCARD TEXT("*ShaderFormat*")
 
+
 /**
  * IShaderFormat, shader pre-compilation abstraction
  */
@@ -19,34 +21,32 @@ public:
 	/** 
 	 * Compile the specified shader.
 	 *
-	 * @param Format - The desired format
-	 * @param Input - The input to the shader compiler.
-	 * @param Output - The output from shader compiler.
-	 * @param WorkingDirectory - The working directory.
+	 * @param Format The desired format
+	 * @param Input The input to the shader compiler.
+	 * @param Output The output from shader compiler.
+	 * @param WorkingDirectory The working directory.
 	 */
-	virtual void CompileShader(FName Format, const struct FShaderCompilerInput& Input, struct FShaderCompilerOutput& Output,const FString& WorkingDirectory) const = 0;
+	virtual void CompileShader( FName Format, const struct FShaderCompilerInput& Input, struct FShaderCompilerOutput& Output,const FString& WorkingDirectory ) const = 0;
 
 	/**
 	 * Gets the current version of the specified shader format.
 	 *
-	 * @param Format - The format to get the version for.
-	 *
+	 * @param Format The format to get the version for.
 	 * @return Version number.
 	 */
-	virtual uint16 GetVersion(FName Format) const = 0;
+	virtual uint16 GetVersion( FName Format ) const = 0;
 
 	/**
 	 * Gets the list of supported formats.
 	 *
-	 * @param OutFormats - Will hold the list of formats.
+	 * @param OutFormats Will hold the list of formats.
 	 */
-	virtual void GetSupportedFormats(TArray<FName>& OutFormats) const = 0;
-
+	virtual void GetSupportedFormats( TArray<FName>& OutFormats ) const = 0;
 
 public:
 
 	/**
 	 * Virtual destructor.
 	 */
-	virtual ~IShaderFormat() { }
+	virtual ~IShaderFormat( ) { }
 };
diff --git a/Engine/Source/Developer/TargetPlatform/Public/Interfaces/IShaderFormatModule.h b/Engine/Source/Developer/TargetPlatform/Public/Interfaces/IShaderFormatModule.h
--- a/Engine/Source/Developer/TargetPlatform/Public/Interfaces/IShaderFormatModule.h
+++ b/Engine/Source/Developer/TargetPlatform/Public/Interfaces/IShaderFormatModule.h
@@ -18,10 +18,12 @@ public:
 	/**
 	 * Gets the shader format.
 	 */
-	virtual IShaderFormat* GetShaderFormat() = 0;
+	virtual IShaderFormat* GetShaderFormat( ) = 0;
 
+public:
 
-protected:
-
-	IShaderFormatModule() { }
+	/**
+	 * Virtual destructor.
+	 */
+	~IShaderFormatModule( ) { }
 };
diff --git a/Engine/Source/Developer/TargetPlatform/Public/Interfaces/ITargetDevice.h b/Engine/Source/Developer/TargetPlatform/Public/Interfaces/ITargetDevice.h
--- a/Engine/Source/Developer/TargetPlatform/Public/Interfaces/ITargetDevice.h
+++ b/Engine/Source/Developer/TargetPlatform/Public/Interfaces/ITargetDevice.h
@@ -14,29 +14,19 @@ namespace ETargetDeviceFeatures
 	 */
 	enum Type
 	{
-		/**
-		 * Multiple instances of a game can run at the same time.
-		 */
+		/** Multiple instances of a game can run at the same time. */
 		MultiLaunch,
 
-		/**
-		 * The device can be powered off remotely.
-		 */
+		/** The device can be powered off remotely. */
 		PowerOff,
 
-		/**
-		 * The device can be powered on remotely.
-		 */
+		/** The device can be powered on remotely. */	  
 		PowerOn,
 
-		/**
-		 * Snapshot of processes running on the device.
-		 */
+		/** Snapshot of processes running on the device. */
 		ProcessSnapshot,
 
-		/**
-		 * The device can be rebooted remotely.
-		 */
+		/** The device can be rebooted remotely. */
 		Reboot
 	};
 }
@@ -49,34 +39,22 @@ namespace ETargetDeviceTypes
 	 */
 	enum Type
 	{
-		/**
-		 * Indeterminate device type.
-		 */
+		/** Indeterminate device type. */
 		Indeterminate,
 
-		/**
-		 * The device is a web browser (i.e. Flash).
-		 */
+		/** The device is a web browser (i.e. Flash). */
 		Browser,
 
-		/**
-		 * The device is a game console.
-		 */
+		/** The device is a game console. */
 		Console,
 
-		/**
-		 * The device is a desktop computer.
-		 */
+		/** The device is a desktop computer. */
 		Desktop,
 
-		/**
-		 * The device is a smart phone.
-		 */
+		/** The device is a smart phone. */
 		Phone,
 
-		/**
-		 * The device is a tablet computer.
-		 */
+		/** The device is a tablet computer. */
 		Tablet
 	};
 
@@ -84,8 +62,7 @@ namespace ETargetDeviceTypes
 	/**
 	 * Returns the string representation of the specified ETargetDeviceTypes value.
 	 *
-	 * @param Configuration - The value to get the string for.
-	 *
+	 * @param Configuration The value to get the string for.
 	 * @return A string value.
 	 */
 	inline FString ToString(ETargetDeviceTypes::Type DeviceType)
@@ -121,31 +98,22 @@ namespace ETargetDeviceThreadStates
 	 */
 	enum Type
 	{
+		/** Unknown thread state. */
 		Unknown,
 
-		/**
-		 * The thread can run, but is not running right now.
-		 */
+		/** The thread can run, but is not running right now. */
 		CanRun,
 
-		/**
-		 * The thread is inactive, i.e. has just been created or exited.
-		 */
+		/** The thread is inactive, i.e. has just been created or exited. */
 		Inactive,
 
-		/**
-		 * The thread cannot run right now.
-		 */
+		/** The thread cannot run right now. */
 		Inhibited,
 
-		/**
-		 * The thread is in the run queue.
-		 */
+		/** The thread is in the run queue. */
 		RunQueue,
 
-		/**
-		 * The thread is running.
-		 */
+		/** The thread is running. */
 		Running
 	};
 }
@@ -155,31 +123,22 @@ namespace ETargetDeviceThreadWaitStates
 {
 	enum Type
 	{
+		/** Unknown wait state. */
 		Unknown,
 
-		/**
-		 * The thread is blocked by a lock.
-		 */
+		/** The thread is blocked by a lock. */
 		Locked,
 	
-		/**
-		 * The thread is sleeping.
-		 */
+		/** The thread is sleeping. */
 		Sleeping,
 
-		/**
-		 * The thread is suspended.
-		 */
+		/** The thread is suspended. */
 		Suspended,
 
-		/**
-		 * The thread is swapped.
-		 */
+		/** The thread is swapped. */
 		Swapped,
 
-		/**
-		 * The thread is waiting on an interrupt.
-		 */
+		/** The thread is waiting on an interrupt. */
 		Waiting
 	};
 }
@@ -190,34 +149,22 @@ namespace ETargetDeviceThreadWaitStates
  */
 struct FTargetDeviceThreadInfo
 {
-	/**
-	 * Holds the exit code.
-	 */
+	/** Holds the exit code. */
 	uint64 ExitCode;
 
-	/**
-	 * Holds the thread identifier.
-	 */
+	/** Holds the thread identifier. */
 	uint32 Id;
 
-	/**
-	 * Holds the name of the thread.
-	 */
+	/** Holds the name of the thread. */
 	FString Name;
 
-	/**
-	 * Holds the thread's stack size.
-	 */
+	/** Holds the thread's stack size. */
 	uint64 StackSize;
 
-	/**
-	 * Holds the thread's current state.
-	 */
+	/** Holds the thread's current state. */
 	ETargetDeviceThreadStates::Type State;
 
-	/**
-	 * Holds the thread's current wait state.
-	 */
+	/** Holds the thread's current wait state. */
 	ETargetDeviceThreadWaitStates::Type WaitState;
 };
 
@@ -227,46 +174,30 @@ struct FTargetDeviceThreadInfo
  */
 struct FTargetDeviceProcessInfo
 {
-	/**
-	 * Holds the process identifier.
-	 */
+	/** Holds the process identifier. */
 	int32 Id;
 
-	/**
-	 * Holds the process name.
-	 */
+	/** Holds the process name. */
 	FString Name;
 
-	/**
-	 * Holds the identifier of the parent process.
-	 */
+	/** Holds the identifier of the parent process. */
 	uint32 ParentId;
 
-	/**
-	 * Holds the collection of threads that belong to this process.
-	 */
+	/** Holds the collection of threads that belong to this process. */
 	TArray<FTargetDeviceThreadInfo> Threads;
 
-	/**
-	 * The name of the user that owns this process.
-	 */
+	/** The name of the user that owns this process. */
 	FString UserName;
 };
 
 
-/**
- * Type definition for shared pointers to instances of ITargetDevice.
- */
+/** Type definition for shared pointers to instances of ITargetDevice. */
 typedef TSharedPtr<class ITargetDevice, ESPMode::ThreadSafe> ITargetDevicePtr;
 
-/**
- * Type definition for shared references to instances of ITargetDevice.
- */
+/** Type definition for shared references to instances of ITargetDevice. */
 typedef TSharedRef<class ITargetDevice, ESPMode::ThreadSafe> ITargetDeviceRef;
 
-/**
- * Type definition for weak pointers to instances of ITargetDevice.
- */
+/** Type definition for weak pointers to instances of ITargetDevice. */
 typedef TWeakPtr<class ITargetDevice, ESPMode::ThreadSafe> ITargetDeviceWeakPtr;
 
 
@@ -287,9 +218,8 @@ public:
 	/**
 	 * Deploys an application in the specified folder to the device.
 	 *
-	 * @param SourceFolder - The path to the files and directories to be deployed.
-	 * @param OutAppId - Will hold the identifier of the deployed application (used for launching).
-	 *
+	 * @param SourceFolder The path to the files and directories to be deployed.
+	 * @param OutAppId Will hold the identifier of the deployed application (used for launching).
 	 * @return true on success, false otherwise.
 	 */
 	virtual bool Deploy( const FString& SourceFolder, FString& OutAppId ) = 0;
@@ -310,7 +240,6 @@ public:
 	 * Gets the unique device identifier.
 	 *
 	 * @return Device identifier.
-	 *
 	 * @see GetName
 	 */
 	virtual FTargetDeviceId GetId( ) const = 0;
@@ -324,7 +253,6 @@ public:
 	 * other string identifying the device that does not need to be unique.
 	 *
 	 * @return Device name.
-	 *
 	 * @see GetId
 	 */
 	virtual FString GetName( ) const = 0;
@@ -339,8 +267,7 @@ public:
 	/**
 	 * Creates a snapshot of processes currently running on the device.
 	 *
-	 * @param OutProcessInfos - Will contain the information for running processes.
-	 *
+	 * @param OutProcessInfos Will contain the information for running processes.
 	 * @return The number of returned processes.
 	 */
 	virtual int32 GetProcessSnapshot( TArray<FTargetDeviceProcessInfo>& OutProcessInfos ) = 0;
@@ -369,12 +296,11 @@ public:
 	/**
 	 * Launches a previously deployed build.
 	 *
-	 * @param AppId - The identifier of the application to launch (as returned by the Deploy() method).
-	 * @param BuildConfiguration - The build configuration to launch.
-	 * @param BuildTarget - The build target type to launch
-	 * @param Params - The command line parameters to launch with.
-	 * @param OutProcessId - Will hold the identifier of the created process (can be NULL).
-	 *
+	 * @param AppId The identifier of the application to launch (as returned by the Deploy() method).
+	 * @param BuildConfiguration The build configuration to launch.
+	 * @param BuildTarget The build target type to launch
+	 * @param Params The command line parameters to launch with.
+	 * @param OutProcessId Will hold the identifier of the created process (can be NULL).
 	 * @return true on success, false otherwise.
 	 */
 	virtual bool Launch( const FString& AppId, EBuildConfigurations::Type BuildConfiguration, EBuildTargets::Type BuildTarget, const FString& Params, uint32* OutProcessId ) = 0;
@@ -382,8 +308,7 @@ public:
 	/**
 	 * Powers off the device.
 	 *
-	 * @param Force - Whether to force powering off.
-	 *
+	 * @param Force Whether to force powering off.
 	 * @return true if the device will be powered off, false otherwise.
 	 */
 	virtual bool PowerOff( bool Force ) = 0;
@@ -398,8 +323,7 @@ public:
 	/** 
 	 * Reboot the device.
 	 *
-	 * @param bReconnect - If true, wait and reconnect when done.
-	 *
+	 * @param bReconnect If true, wait and reconnect when done.
 	 * @return true if the reboot was successful from the perspective of the PC .
 	 */
 	virtual bool Reboot( bool bReconnect = false ) = 0;
@@ -407,10 +331,9 @@ public:
 	/**
 	 * Runs an executable on the device.
 	 *
-	 * @param ExecutablePath - The path to the executable to run.
-	 * @param Params - The command line parameters.
-	 * @param OutProcessId - Will hold the identifier of the created process (can be NULL).
-	 *
+	 * @param ExecutablePath The path to the executable to run.
+	 * @param Params The command line parameters.
+	 * @param OutProcessId Will hold the identifier of the created process (can be NULL).
 	 * @return true if the executable was started, false otherwise.
 	 */
 	virtual bool Run( const FString& ExecutablePath, const FString& Params, uint32* OutProcessId ) = 0;
@@ -418,8 +341,7 @@ public:
 	/**
 	 * Checks whether the target device supports the specified feature.
 	 *
-	 * @param Feature - The feature to check.
-	 *
+	 * @param Feature The feature to check.
 	 * @return true if the feature is supported, false otherwise.
 	 */
 	virtual bool SupportsFeature( ETargetDeviceFeatures::Type Feature ) const = 0;
@@ -427,8 +349,7 @@ public:
 	/**
 	 * Checks whether this device supports the specified SDK version.
 	 *
-	 * @param VersionString - The SDK version string.
-	 *
+	 * @param VersionString The SDK version string.
 	 * @return true if the SDK version is supported, false otherwise.
 	 */
 	virtual bool SupportsSdkVersion( const FString& VersionString ) const = 0;
@@ -436,8 +357,7 @@ public:
 	/**
 	 * Terminates a process that was launched on the device using the Launch() or Run() methods.
 	 *
-	 * @param ProcessId - The identifier of the process to terminate.
-	 *
+	 * @param ProcessId The identifier of the process to terminate.
 	 * @return true if the process was terminated, false otherwise.
 	 */
 	virtual bool TerminateProcess( const int32 ProcessId ) = 0;
@@ -445,25 +365,24 @@ public:
 	/**
 	 * Set credentials for the user account to use on the device
 	 * 
-	 * @param UserName - The user account on the device we will run under
-	 * @param UserPassword - The password for the user account on the device we will run under
+	 * @param UserName The user account on the device we will run under
+	 * @param UserPassword The password for the user account on the device we will run under.
 	 */
 	virtual void SetUserCredentials( const FString & UserName, const FString & UserPassword ) = 0;
 
 	/**
 	 * Get credentials for the user account to use on the device
 	 * 
-	 * @param OutUserName - The user account on the device we will run under
-	 * @param OutUserPassword - The password for the user account on the device we will run under
-	 *
-	 * @return false if not supported
+	 * @param OutUserName The user account on the device we will run under
+	 * @param OutUserPassword The password for the user account on the device we will run under
+	 * @return true on success, false if not supported.
 	 */
 	virtual bool GetUserCredentials( FString & OutUserName, FString & OutUserPassword ) = 0;
 
 public:
 	
 	/**
-	 * Destructor.
+	 * Virtual destructor.
 	 */
 	virtual ~ITargetDevice( ) { }
 };
diff --git a/Engine/Source/Developer/TargetPlatform/Public/Interfaces/ITargetPlatform.h b/Engine/Source/Developer/TargetPlatform/Public/Interfaces/ITargetPlatform.h
--- a/Engine/Source/Developer/TargetPlatform/Public/Interfaces/ITargetPlatform.h
+++ b/Engine/Source/Developer/TargetPlatform/Public/Interfaces/ITargetPlatform.h
@@ -14,9 +14,7 @@ namespace ETargetPlatformBuildArtifacts
 	 */
 	enum Type
 	{
-		/**
-		 * Include content files.
-		 */
+		/** Include content files. */
 		Content = 0x1,
 
 		/**
@@ -26,14 +24,10 @@ namespace ETargetPlatformBuildArtifacts
 		 */
 		DebugSymbols = 0x2,
 
-		/**
-		 * Include Engine binaries and DLLs.
-		 */
+		/** Include Engine binaries and DLLs. */
 		Engine = 0x4,
 
-		/**
-		 * Include tools.
-		 */
+		/** Include tools. */
 		Tools = 0x8,
 	};
 }
@@ -46,59 +40,37 @@ namespace ETargetPlatformFeatures
 	 */
 	enum Type
 	{
-		/**
-		 * Distance field shadows.
-		 */
+		/** Distance field shadows. */
 		DistanceFieldShadows,
 
-		/**
-		 * Gray scale SRGB texture formats support.
-		 */
+		/** Gray scale SRGB texture formats support. */
 		GrayscaleSRGB,
 
-		/**
-		 * High quality light maps.
-		 */
+		/** High quality light maps. */
 		HighQualityLightmaps,
 
-		/**
-		 * Low quality light maps.
-		 */
+		/** Low quality light maps. */
 		LowQualityLightmaps,
 
-		/**
-		 * Run multiple game instances on a single device.
-		 */
+		/** Run multiple game instances on a single device. */
 		MultipleGameInstances,
 
-		/**
-		 * Builds can be packaged for this platform.
-		 */
+		/** Builds can be packaged for this platform. */
 		Packaging,
 
-		/**
-		 * Connect and disconnect devices through the SDK.
-		 */
+		/** Connect and disconnect devices through the SDK. */
 		SdkConnectDisconnect,
 
-		/**
-		 * GPU tesselation.
-		 */
+		/** GPU tesselation. */
 		Tessellation,
 
-		/**
-		 * Texture streaming.
-		 */
+		/** Texture streaming. */
 		TextureStreaming,
 
-		/**
-		 * User credentials are required to use the device.
-		 */
+		/** User credentials are required to use the device. */
 		UserCredentials,
 
-		/**
-		 * Vertex Shader Texture Sampling.
-		 */
+		/** Vertex Shader Texture Sampling. */
 		VertexShaderTextureSampling,
 	};
 };
@@ -111,19 +83,13 @@ namespace ETargetPlatformIcons
 	 */
 	enum IconType
 	{
-		/**
-		 * Normal sized icon.
-		 */
+		/** Normal sized icon. */
 		Normal,
 
-		/**
-		 * Large icon.
-		 */
+		/** Large icon. */
 		Large,
 
-		/**
-		 * Extra large icon.
-		 */
+		/** Extra large icon. */
 		XLarge
 	};
 }
@@ -139,9 +105,8 @@ public:
 	/**
 	 * Add a target device by name.
 	 *
-	 * @param DeviceName - The name of the device to add.
-	 * @param bDefault - Whether the added device should be the default.
-	 *
+	 * @param DeviceName The name of the device to add.
+	 * @param bDefault Whether the added device should be the default.
 	 * @return true if the device was added, false otherwise.
 	 */
 	virtual bool AddDevice( const FString& DeviceName, bool bDefault ) = 0;
@@ -168,7 +133,7 @@ public:
 	/**
 	 * Returns all discoverable physical devices.
 	 *
-	 * @param OutDevices - Will contain a list of discovered devices.
+	 * @param OutDevices Will contain a list of discovered devices.
 	 */
 	virtual void GetAllDevices( TArray<ITargetDevicePtr>& OutDevices ) const = 0;
 
@@ -182,9 +147,8 @@ public:
 	/** 
 	 * Generates a platform specific asset manifest given an array of FAssetData.
 	 *
-	 * @param ChunkMap - A map of asset path to ChunkIDs for all of the assets
-	 * @param ChunkIDsInUse - A set of all ChunkIDs used by this set of assets
-	 *
+	 * @param ChunkMap A map of asset path to ChunkIDs for all of the assets
+	 * @param ChunkIDsInUse A set of all ChunkIDs used by this set of assets
 	 * @return true if the manifest was successfully generated, or if the platform doesn't need a manifest 
 	 */
 	virtual bool GenerateStreamingInstallManifest( const TMultiMap<FString, int32>& ChunkMap, const TSet<int32>& ChunkIDsInUse ) const = 0;
@@ -201,17 +165,15 @@ public:
 	/** 
 	 * Gets an interface to the specified device.
 	 *
-	 * @param DeviceId - The identifier of the device to get.
-	 *
-	 * @return The target device (can be NULL). 
+	 * @param DeviceId The identifier of the device to get.
+	 * @return The target device (can be nullptr). 
 	 */
 	virtual ITargetDevicePtr GetDevice( const FTargetDeviceId& DeviceId ) = 0;
 
 	/**
 	 * Gets the path to the platform's icon.
 	 *
-	 * @param IconType - The type of icon to get (i.e. Small or Large).
-	 *
+	 * @param IconType The type of icon to get (i.e. Small or Large).
 	 * @return The path to the icon.
 	 */
 	virtual FString GetIconPath( ETargetPlatformIcons::IconType IconType ) const = 0;
@@ -265,7 +227,6 @@ public:
 	 * Returns the name of this platform
 	 *
 	 * @return Platform name.
-	 *
 	 * @see DisplayName
 	 */
 	virtual FString PlatformName( ) const = 0;
@@ -288,7 +249,6 @@ public:
 	 * Checks whether this platform supports the specified build target, i.e. Game or Editor.
 	 *
 	 * @param BuildTarget - The build target to check.
-	 *
 	 * @return true if the build target is supported, false otherwise.
 	 */
 	virtual bool SupportsBuildTarget( EBuildTargets::Type BuildTarget ) const = 0;
@@ -297,7 +257,6 @@ public:
 	 * Checks whether the target platform supports the specified feature.
 	 *
 	 * @param Feature - The feature to check.
-	 *
 	 * @return true if the feature is supported, false otherwise.
 	 */
 	virtual bool SupportsFeature( ETargetPlatformFeatures::Type Feature ) const = 0;
@@ -307,9 +266,8 @@ public:
 	 * Checks whether the platform's SDK requirements are met so that we can do things like
 	 * package for the platform
 	 *
-	 * @param bProjectHasCode - true if the project has code, and therefore any compilation based SDK requirements should be checked
-	 * @param OutDocumentationPath - Let's the platform tell the editor a path to show some documentation about how to set up the SDK
-	 *
+	 * @param bProjectHasCode true if the project has code, and therefore any compilation based SDK requirements should be checked
+	 * @param OutDocumentationPath Let's the platform tell the editor a path to show some documentation about how to set up the SDK
 	 * @return true if the platform is ready for use
 	 */
 	virtual bool IsSdkInstalled(bool bProjectHasCode, FString& OutDocumentationPath) const = 0;
@@ -325,37 +283,36 @@ public:
 	/**
 	 * Gets the reflection capture formats this platform needs.
 	 *
-	 * @param OutFormats - Will contain the collection of formats.
+	 * @param OutFormats Will contain the collection of formats.
 	 */
 	virtual void GetReflectionCaptureFormats( TArray<FName>& OutFormats ) const = 0;
 
 	/**
 	 * Gets the shader formats this platform can use.
 	 *
-	 * @param OutFormats - Will contain the shader formats.
+	 * @param OutFormats Will contain the shader formats.
 	 */
 	virtual void GetAllPossibleShaderFormats( TArray<FName>& OutFormats ) const = 0;
 
 	/**
 	* Gets the shader formats that have been selected for this target platform
 	*
-	* @param OutFormats - Will contain the shader formats.
+	* @param OutFormats Will contain the shader formats.
 	*/
 	virtual void GetAllTargetedShaderFormats(TArray<FName>& OutFormats) const = 0;
 
 	/**
 	 * Gets the format to use for a particular texture.
 	 *
-	 * @param Texture - The texture to get the format for.
-	 * @param OutFormats - Will contain the list of supported formats.
+	 * @param Texture The texture to get the format for.
+	 * @param OutFormats Will contain the list of supported formats.
 	 */
 	virtual void GetTextureFormats( const class UTexture* Texture, TArray<FName>& OutFormats ) const = 0;
 
 	/**
 	 * Gets the format to use for a particular piece of audio.
 	 *
-	 * @param Wave - The sound node wave to get the format for.
-	 *
+	 * @param Wave The sound node wave to get the format for.
 	 * @return Name of the wave format.
 	 */
 	virtual FName GetWaveFormat( class USoundWave* Wave ) const = 0;
@@ -378,9 +335,8 @@ public:
 	/** 
 	 * Package a build for the given platform 
 	 * 
-	 * @param InPackgeDirectory The directory that contains what needs to be packaged 
-	 * 
-	 * @return bool True if successful, false if failed 
+	 * @param InPackgeDirectory The directory that contains what needs to be packaged.
+	 * @return bool true on success, false otherwise.
 	 */ 
 	virtual bool PackageBuild(const FString& InPackgeDirectory) = 0;
 
diff --git a/Engine/Source/Developer/TargetPlatform/Public/Interfaces/ITargetPlatformManagerModule.h b/Engine/Source/Developer/TargetPlatform/Public/Interfaces/ITargetPlatformManagerModule.h
--- a/Engine/Source/Developer/TargetPlatform/Public/Interfaces/ITargetPlatformManagerModule.h
+++ b/Engine/Source/Developer/TargetPlatform/Public/Interfaces/ITargetPlatformManagerModule.h
@@ -15,65 +15,53 @@ class ITargetPlatformManagerModule
 {
 public:
 
-	/** 
-	 * Invalidate should be called if any TargetPlatform modules get loaded/unloaded/reloaded during 
-	 * runtime to give the implementation the chance to rebuild all its internal states and caches.
-	 **/
-	virtual void Invalidate() = 0;
-
 	/**
 	 * Finds an audio format with the specified name.
 	 *
-	 * @param Name - Name of the format to find.
-	 *
-	 * @return The audio format, or NULL if not found.
+	 * @param Name Name of the format to find.
+	 * @return The audio format, or nullptr if not found.
 	 */
-	virtual const class IAudioFormat* FindAudioFormat(FName Name) = 0;
+	virtual const class IAudioFormat* FindAudioFormat( FName Name ) = 0;
 
 	/**
 	 * Finds a PhysX format with the specified name.
 	 *
-	 * @param Name - Name of the format to find.
-	 *
-	 * @return The PhysX format, or NULL if not found.
+	 * @param Name Name of the format to find.
+	 * @return The PhysX format, or nullptr if not found.
 	 */
-	virtual const class IPhysXFormat* FindPhysXFormat(FName Name) = 0;
+	virtual const class IPhysXFormat* FindPhysXFormat( FName Name ) = 0;
 
 	/**
 	 * Finds a shader format with the specified name.
 	 *
-	 * @param Name - Name of the format to find.
-	 *
-	 * @return The shader format, or NULL if not found.
+	 * @param Name Name of the format to find.
+	 * @return The shader format, or nullptr if not found.
 	 */
-	virtual const class IShaderFormat* FindShaderFormat(FName Name) = 0;
+	virtual const class IShaderFormat* FindShaderFormat( FName Name ) = 0;
 
 	/**
 	 * Finds a texture format with the specified name.
 	 *
-	 * @param Name - Name of the format to find.
-	 *
-	 * @return The texture format, or NULL if not found.
+	 * @param Name Name of the format to find.
+	 * @return The texture format, or nullptr if not found.
 	 */
-	virtual const class ITextureFormat* FindTextureFormat(FName Name) = 0;
+	virtual const class ITextureFormat* FindTextureFormat( FName Name ) = 0;
 
 	/**
 	 * Finds a target device by identifier.
 	 *
-	 * @param DeviceId - The device identifier.
-	 *
-	 * @return The target device, or NULL if not found.
+	 * @param DeviceId The device identifier.
+	 * @return The target device, or nullptr if not found.
 	 */
-	virtual ITargetDevicePtr FindTargetDevice(const FTargetDeviceId& DeviceId) = 0;
+	virtual ITargetDevicePtr FindTargetDevice( const FTargetDeviceId& DeviceId ) = 0;
 
 	/**
 	 * Finds a target platform by name.
 	 *
-	 * @param Name - The target platform's short or long name.
-	 *
-	 * @return The target platform, or NULL if not found.
+	 * @param Name The target platform's short or long name.
+	 * @return The target platform, or nullptr if not found.
 	 */
-	virtual ITargetPlatform* FindTargetPlatform(FString Name) = 0;
+	virtual ITargetPlatform* FindTargetPlatform( FString Name ) = 0;
 
 
 	/**
@@ -81,28 +69,28 @@ public:
 	 *
 	 * @return Collection of platforms.
 	 */
-	virtual const TArray<ITargetPlatform*>& GetCookingTargetPlatforms() = 0;
+	virtual const TArray<ITargetPlatform*>& GetCookingTargetPlatforms( ) = 0;
 
 	/**
 	 * Return the list of the ITargetPlatforms that we want to build data for.
 	 *
 	 * @return Collection of platforms.
 	 */
-	virtual const TArray<ITargetPlatform*>& GetActiveTargetPlatforms() = 0;
+	virtual const TArray<ITargetPlatform*>& GetActiveTargetPlatforms( ) = 0;
 
 	/**
 	 * Returns the list of all IAudioFormats that were located in DLLs.
 	 *
 	 * @return Collection of audio formats.
 	 */
-	virtual const TArray<const class IAudioFormat*>& GetAudioFormats() = 0;
+	virtual const TArray<const class IAudioFormat*>& GetAudioFormats( ) = 0;
 
 	/**
 	 * Returns the list of all IPhysXFormats that were located in DLLs.
 	 *
 	 * @return Collection of PhysX formats.
 	 */
-	virtual const TArray<const class IPhysXFormat*>& GetPhysXFormats() = 0;
+	virtual const TArray<const class IPhysXFormat*>& GetPhysXFormats( ) = 0;
 
 	/**
 	 * Returns the target platform that is currently running.
@@ -111,42 +99,56 @@ public:
 	 *
 	 * @return Running target platform.
 	 */
-	virtual ITargetPlatform* GetRunningTargetPlatform() = 0;
+	virtual ITargetPlatform* GetRunningTargetPlatform( ) = 0;
 
 	/**
 	 * Returns the list of all ITextureFormats that were located in DLLs.
 	 *
 	 * @return Collection of shader formats.
 	 */
-	virtual const TArray<const class IShaderFormat*>& GetShaderFormats() = 0;
+	virtual const TArray<const class IShaderFormat*>& GetShaderFormats( ) = 0;
 
 	/**
 	 * Returns the list of all ITargetPlatforms that were located in DLLs.
 	 *
 	 * @return Collection of platforms.
 	 */
-	virtual const TArray<ITargetPlatform*>& GetTargetPlatforms() = 0;
+	virtual const TArray<ITargetPlatform*>& GetTargetPlatforms( ) = 0;
 
 	/**
 	 * Returns the list of all ITextureFormats that were located in DLLs.
 	 *
 	 * @return Collection of texture formats.
 	 */
-	virtual const TArray<const class ITextureFormat*>& GetTextureFormats() = 0;
+	virtual const TArray<const class ITextureFormat*>& GetTextureFormats( ) = 0;
 	
+	/**
+	 * Invalidates the target platform module.
+	 *
+	 * Invalidate should be called if any TargetPlatform modules get loaded/unloaded/reloaded during 
+	 * runtime to give the implementation the chance to rebuild all its internal states and caches.
+	 */
+	virtual void Invalidate( ) = 0;
+
 	/**
 	 * Checks whether we should only build formats that are actually required for use by the runtime.
 	 *
 	 * @return true if formats are restricted, false otherwise.
 	 */
-	virtual bool RestrictFormatsToRuntimeOnly() = 0;
+	virtual bool RestrictFormatsToRuntimeOnly( ) = 0;
 
 	/**
 	 * Gets the shader format version for the specified shader.
 	 *
-	 * @param Name - Name of the shader format to get the version for.
-	 *
+	 * @param Name Name of the shader format to get the version for.
 	 * @return Version number.
 	 */
 	virtual uint16 ShaderFormatVersion(FName Name) = 0;
+
+public:
+
+	/**
+	 * Virtual destructor.
+	 */
+	~ITargetPlatformManagerModule() { }
 };
diff --git a/Engine/Source/Developer/TargetPlatform/Public/Interfaces/ITargetPlatformModule.h b/Engine/Source/Developer/TargetPlatform/Public/Interfaces/ITargetPlatformModule.h
--- a/Engine/Source/Developer/TargetPlatform/Public/Interfaces/ITargetPlatformModule.h
+++ b/Engine/Source/Developer/TargetPlatform/Public/Interfaces/ITargetPlatformModule.h
@@ -18,10 +18,12 @@ public:
 	/**
 	 * Gets the target platform for this module.
 	 */
-	virtual ITargetPlatform* GetTargetPlatform() = 0;
+	virtual ITargetPlatform* GetTargetPlatform( ) = 0;
 
+public:
 
-protected:
-
-	ITargetPlatformModule() { }
+	/**
+	 * Virtual destructor.
+	 */
+	~ITargetPlatformModule( ) { }
 };
diff --git a/Engine/Source/Developer/TargetPlatform/Public/Interfaces/ITextureFormat.h b/Engine/Source/Developer/TargetPlatform/Public/Interfaces/ITextureFormat.h
--- a/Engine/Source/Developer/TargetPlatform/Public/Interfaces/ITextureFormat.h
+++ b/Engine/Source/Developer/TargetPlatform/Public/Interfaces/ITextureFormat.h
@@ -6,15 +6,20 @@
 
 #pragma once
 
+
+/**
+ * Structure for texture format compressor capabilities.
+ */
 struct FTextureFormatCompressorCaps
 {
-	FTextureFormatCompressorCaps() : MaxTextureDimension(TNumericLimits<uint32>::Max())
-	{
-	}
+	FTextureFormatCompressorCaps( )
+		: MaxTextureDimension(TNumericLimits<uint32>::Max())
+	{ }
 
 	uint32 MaxTextureDimension;
 };
 
+
 /**
  * Interface for texture compression modules.
  */
@@ -27,7 +32,7 @@ public:
 	 *
 	 * @return true if parallel compression is supported, false otherwise.
 	 */
-	virtual bool AllowParallelBuild() const
+	virtual bool AllowParallelBuild( ) const
 	{
 		return false;
 	}
@@ -35,21 +40,19 @@ public:
 	/**
 	 * Gets the current version of the specified texture format.
 	 *
-	 * @param Format - The format to get the version for.
-	 *
+	 * @param Format The format to get the version for.
 	 * @return Version number.
 	 */
-	virtual uint16 GetVersion(FName Format) const = 0;
+	virtual uint16 GetVersion( FName Format ) const = 0;
 
 	/**
 	 * Gets an optional derived data key string, so that the compressor can
 	 * rely upon the number of mips, size of texture, etc, when compressing the image
 	 *
-	 * @param Texture - Reference to the texture we are compressing
-	 *
+	 * @param Texture Reference to the texture we are compressing.
 	 * @return A string that will be used with the DDC, the string should be in the format "<DATA>_"
 	 */
-	virtual FString GetDerivedDataKeyString(const class UTexture& Texture) const
+	virtual FString GetDerivedDataKeyString( const class UTexture& Texture ) const
 	{
 		return TEXT("");
 	}
@@ -59,16 +62,14 @@ public:
 	 *
 	 * @param OutFormats - Will hold the list of formats.
 	 */
-	virtual void GetSupportedFormats(TArray<FName>& OutFormats) const = 0;
-
-
+	virtual void GetSupportedFormats( TArray<FName>& OutFormats ) const = 0;
 
 	/**
 	* Gets the capabilities of the texture compressor.
 	*
 	* @param OutCaps - filled with capability properties of texture format compressor.
 	*/
-	virtual FTextureFormatCompressorCaps GetFormatCapabilities() const = 0;
+	virtual FTextureFormatCompressorCaps GetFormatCapabilities( ) const = 0;
 
 	/**
 	 * Compresses a single image.
@@ -93,5 +94,5 @@ public:
 	/**
 	 * Virtual destructor.
 	 */
-	virtual ~ITextureFormat() { }
+	virtual ~ITextureFormat( ) { }
 };
diff --git a/Engine/Source/Developer/TargetPlatform/Public/Interfaces/ITextureFormatModule.h b/Engine/Source/Developer/TargetPlatform/Public/Interfaces/ITextureFormatModule.h
--- a/Engine/Source/Developer/TargetPlatform/Public/Interfaces/ITextureFormatModule.h
+++ b/Engine/Source/Developer/TargetPlatform/Public/Interfaces/ITextureFormatModule.h
@@ -18,10 +18,12 @@ public:
 	/**
 	 * Gets the texture format.
 	 */
-	virtual ITextureFormat* GetTextureFormat() = 0;
+	virtual ITextureFormat* GetTextureFormat( ) = 0;
 
+public:
 
-protected:
-
-	ITextureFormatModule() { }
+	/**
+	 * Virtual destructor.
+	 */
+	~ITextureFormatModule() { }
 };
diff --git a/Engine/Source/Developer/TargetPlatform/Public/Interfaces/TargetDeviceId.h b/Engine/Source/Developer/TargetPlatform/Public/Interfaces/TargetDeviceId.h
--- a/Engine/Source/Developer/TargetPlatform/Public/Interfaces/TargetDeviceId.h
+++ b/Engine/Source/Developer/TargetPlatform/Public/Interfaces/TargetDeviceId.h
@@ -22,22 +22,20 @@ public:
 	/**
 	 * Creates and initializes a new target device identifier.
 	 *
-	 * @param InPlatformName - The name of the device's target platform.
-	 * @param InDeviceName - The unique device name, i.e. IP address or computer name.
+	 * @param InPlatformName The name of the device's target platform.
+	 * @param InDeviceName The unique device name, i.e. IP address or computer name.
 	 */
 	FTargetDeviceId( const FString& InPlatformName, const FString& InDeviceName )
 		: PlatformName(InPlatformName)
 		, DeviceName(InDeviceName)
 	{ }
 
-
 public:
 
 	/**
 	 * Compares this device identifier with the given device identifier for equality.
 	 *
-	 * @param Other - The device identifier to compare with.
-	 *
+	 * @param Other The device identifier to compare with.
 	 * @return true if the identifiers are equal, false otherwise.
 	 */
 	bool operator==( const FTargetDeviceId& Other ) const
@@ -48,8 +46,7 @@ public:
 	/**
 	 * Compares this device identifier with the given device identifier for inequality.
 	 *
-	 * @param Other - The device identifier to compare with.
-	 *
+	 * @param Other The device identifier to compare with.
 	 * @return true if the identifiers are not equal, false otherwise.
 	 */
 	bool operator!=( const FTargetDeviceId& Other ) const
@@ -60,9 +57,8 @@ public:
 	/**
 	 * Serializes the given device identifier from or into the specified archive.
 	 *
-	 * @param Ar - The archive to serialize from or into.
-	 * @param Uri - The device identifier to serialize.
-	 *
+	 * @param Ar The archive to serialize from or into.
+	 * @param Uri The device identifier to serialize.
 	 * @return The archive.
 	 */
 	friend FArchive& operator<<( FArchive& Ar, FTargetDeviceId& Uri )
@@ -70,7 +66,6 @@ public:
 		return Ar << Uri.PlatformName << Uri.DeviceName;
 	}
 
-
 public:
 
 	/**
@@ -113,14 +108,12 @@ public:
 		return PlatformName + TEXT("@") + DeviceName;
 	}
 
-
 public:
 
 	/**
 	 * Gets the hash for specified target device identifier.
 	 *
-	 * @param Uri - The device identifier to get the hash for.
-	 *
+	 * @param Uri The device identifier to get the hash for.
 	 * @return Hash value.
 	 */
 	friend uint32 GetTypeHash( const FTargetDeviceId& Uri )
@@ -128,15 +121,13 @@ public:
 		return GetTypeHash(Uri.ToString());
 	}
 
-
 public:
 
 	/**
 	 * Converts a string to a target device identifier.
 	 *
-	 * @param IdString - The string to convert.
-	 * @param OutId - Will contain the parsed device identifier.
-	 *
+	 * @param IdString The string to convert.
+	 * @param OutId Will contain the parsed device identifier.
 	 * @return true if the string was converted successfully, false otherwise.
 	 */
 	static bool Parse( const FString& IdString, FTargetDeviceId& OutId )
@@ -144,7 +135,6 @@ public:
 		return IdString.Split(TEXT("@"), &OutId.PlatformName, &OutId.DeviceName);
 	}
 
-
 private:
 
 	// Holds the name of the device's target platform.
