better default implementation for IsReadyToBuld
#ue4

diff --git a/Engine/Source/Developer/TargetPlatform/Public/Common/TargetPlatformBase.h b/Engine/Source/Developer/TargetPlatform/Public/Common/TargetPlatformBase.h
--- a/Engine/Source/Developer/TargetPlatform/Public/Common/TargetPlatformBase.h
+++ b/Engine/Source/Developer/TargetPlatform/Public/Common/TargetPlatformBase.h
@@ -172,7 +172,12 @@ public:
 
 	virtual int IsReadyToBuild(const FString& ProjectPath, bool bProjectHasCode, FString& OutDocumentationPath) const override
 	{
-		return ETargetPlatformReadyStatus::Ready;
+		int bReadyToBuild = ETargetPlatformReadyStatus::Ready; // @todo How do we check that the iOS SDK is installed when building from Windows? Is that even possible?
+		if (!IsSdkInstalled(bProjectHasCode, OutDocumentationPath))
+		{
+			bReadyToBuild |= ETargetPlatformReadyStatus::SDKNotFound;
+		}
+		return bReadyToBuild;
 	}
 
 	// End ITargetPlatform interface
