Fix Linux editor build.

diff --git a/Engine/Source/Developer/TargetPlatform/Private/TargetPlatformManagerModule.cpp b/Engine/Source/Developer/TargetPlatform/Private/TargetPlatformManagerModule.cpp
--- a/Engine/Source/Developer/TargetPlatform/Private/TargetPlatformManagerModule.cpp
+++ b/Engine/Source/Developer/TargetPlatform/Private/TargetPlatformManagerModule.cpp
@@ -540,6 +540,7 @@ protected:
 
 	bool SetupAndValidateAutoSDK(ITargetPlatform& Platform)
 	{
+#if AUTOSDKS_ENABLED
 		bool bValidSDK = false;
 		// Invoke UBT to perform SDK switching, or detect that a proper manual SDK is already setup.
 		FString PlatformName = Platform.IniPlatformName();
@@ -557,6 +558,9 @@ protected:
 			PlatformsSetup.Add(PlatformFName, bValidSDK);
 		}
 		return bValidSDK;
+#else
+		return true;
+#endif // AUTOSDKS_ENABLED
 	}
 	
 	bool SetupEnvironmentFromAutoSDK(ITargetPlatform& Platform)
