Platform target dlls for iOS and Android are always built for Editor & Program builds.

ttp 3414656

diff --git a/Engine/Build/JunkManifest.txt b/Engine/Build/JunkManifest.txt
--- a/Engine/Build/JunkManifest.txt
+++ b/Engine/Build/JunkManifest.txt
@@ -13,10 +13,7 @@ Platform=XboxOne:Engine\Binaries\Win64\RocketEditor-*XboxOne*.*
 Platform=XboxOne:Engine\Binaries\Win32\RocketEditor-*XboxOne*.*
 Platform=XboxOne:Engine\Binaries\Win64\ShaderCompileWorker-*XboxOne*.*
 Platform=XboxOne:Engine\Binaries\Win32\ShaderCompileWorker-*XboxOne*.*
-Platform=IOS:Engine\Binaries\Win64\IOS\*.*
 Platform=Linux:Engine\Binaries\Win64\Linux\*.*
-Platform=Android:Engine\Binaries\Win64\Android\*.*
-Platform=HTML5:Engine\Binaries\Win64\HTML5\*.*
 Engine\Binaries\Win64\UE4Editor-LinuxTarget*.*
 Engine\Binaries\Win64\UE4Editor-AndroidTarget*.*
 Engine\Binaries\Win64\UE4Editor-IOSPlatformEditor*.*
diff --git a/Engine/Source/Developer/TargetPlatform/TargetPlatform.Build.cs b/Engine/Source/Developer/TargetPlatform/TargetPlatform.Build.cs
--- a/Engine/Source/Developer/TargetPlatform/TargetPlatform.Build.cs
+++ b/Engine/Source/Developer/TargetPlatform/TargetPlatform.Build.cs
@@ -40,6 +40,18 @@ public class TargetPlatform : ModuleRules
 					DynamicallyLoadedModuleNames.Add("AudioFormatOgg");
 					DynamicallyLoadedModuleNames.Add("AudioFormatOpus"); 
 				}
+
+                if (Target.Type == TargetRules.TargetType.Editor || Target.Type == TargetRules.TargetType.Program)
+                {
+                    DynamicallyLoadedModuleNames.Add("AndroidTargetPlatform");
+                    DynamicallyLoadedModuleNames.Add("Android_PVRTCTargetPlatform");
+                    DynamicallyLoadedModuleNames.Add("Android_ATCTargetPlatform");
+                    DynamicallyLoadedModuleNames.Add("Android_DXTTargetPlatform");
+                    DynamicallyLoadedModuleNames.Add("Android_ETC1TargetPlatform");
+                    DynamicallyLoadedModuleNames.Add("Android_ETC2TargetPlatform");
+                    DynamicallyLoadedModuleNames.Add("IOSTargetPlatform");
+                    DynamicallyLoadedModuleNames.Add("HTML5TargetPlatform");
+                }
 			}
 			else if (Target.Platform == UnrealTargetPlatform.Mac)
 			{
@@ -55,6 +67,18 @@ public class TargetPlatform : ModuleRules
 				{
 					DynamicallyLoadedModuleNames.Add("AudioFormatOgg");
 				}
+
+                if (Target.Type == TargetRules.TargetType.Editor || Target.Type == TargetRules.TargetType.Program)
+                {
+                    DynamicallyLoadedModuleNames.Add("AndroidTargetPlatform");
+                    DynamicallyLoadedModuleNames.Add("Android_PVRTCTargetPlatform");
+                    DynamicallyLoadedModuleNames.Add("Android_ATCTargetPlatform");
+                    DynamicallyLoadedModuleNames.Add("Android_DXTTargetPlatform");
+                    DynamicallyLoadedModuleNames.Add("Android_ETC1TargetPlatform");
+                    DynamicallyLoadedModuleNames.Add("Android_ETC2TargetPlatform");
+                    DynamicallyLoadedModuleNames.Add("IOSTargetPlatform");
+                    DynamicallyLoadedModuleNames.Add("HTML5TargetPlatform");
+                }
 			}
             else if (Target.Platform == UnrealTargetPlatform.Linux)
             {
