Make AutoSDK failures warnings in the cooker so-as not to fail unrelated cooking.

diff --git a/Engine/Source/Developer/TargetPlatform/Private/TargetPlatformManagerModule.cpp b/Engine/Source/Developer/TargetPlatform/Private/TargetPlatformManagerModule.cpp
--- a/Engine/Source/Developer/TargetPlatform/Private/TargetPlatformManagerModule.cpp
+++ b/Engine/Source/Developer/TargetPlatform/Private/TargetPlatformManagerModule.cpp
@@ -629,7 +629,7 @@ protected:
 			
 			if (FileLines.Num() != 2)
 			{
-				UE_LOG(LogTargetPlatformManager, Error, TEXT("Malformed install manifest file for Platform %s"), *AutoSDKPath);				
+				UE_LOG(LogTargetPlatformManager, Warning, TEXT("Malformed install manifest file for Platform %s"), *AutoSDKPath);
 				return false;
 			}
 
@@ -642,7 +642,7 @@ protected:
 		}
 		else
 		{	
-			UE_LOG(LogTargetPlatformManager, Error, TEXT("install manifest file for Platform %s not found."), *AutoSDKPath);			
+			UE_LOG(LogTargetPlatformManager, Warning, TEXT("install manifest file for Platform %s not found.  Platform not set up."), *AutoSDKPath);			
 			return false;			
 		}		
 
@@ -771,7 +771,7 @@ protected:
 		}
 		else
 		{
-			UE_LOG(LogTargetPlatformManager, Error, TEXT("OutputEnvVars.txt not found for platform: '%s'"), *AutoSDKPath);			
+			UE_LOG(LogTargetPlatformManager, Warning, TEXT("OutputEnvVars.txt not found for platform: '%s'"), *AutoSDKPath);			
 			return false;
 		}
 
