fix for parsing platform SDK validity on Mac
#ue4
#mac

diff --git a/Engine/Source/Developer/TargetPlatform/Private/TargetPlatformManagerModule.cpp b/Engine/Source/Developer/TargetPlatform/Private/TargetPlatformManagerModule.cpp
--- a/Engine/Source/Developer/TargetPlatform/Private/TargetPlatformManagerModule.cpp
+++ b/Engine/Source/Developer/TargetPlatform/Private/TargetPlatformManagerModule.cpp
@@ -787,8 +787,14 @@ protected:
 	bool SetupSDKStatus()
 	{
 		// run UBT with -validate -allplatforms and read the output
+#if PLATFORM_MAC
+		FString CmdExe = TEXT("/bin/sh");
+        FString ScriptPath = FPaths::ConvertRelativePathToFull(FPaths::EngineDir() / TEXT("Build/BatchFiles/Mac/RunMono.sh"));
+		FString CommandLine = TEXT("\"") + ScriptPath + TEXT("\" \"") + FPaths::ConvertRelativePathToFull(FPaths::EngineDir() / TEXT("Binaries/DotNet/UnrealBuildTool.exe")) + TEXT("\" -validateplatform -allplatforms");
+#else
 		FString CmdExe = FPaths::ConvertRelativePathToFull(FPaths::EngineDir() / TEXT("Binaries/DotNet/UnrealBuildTool.exe"));
 		FString CommandLine = TEXT("-validateplatform -allplatforms");
+#endif
 		TSharedPtr<FMonitoredProcess> UBTProcess = MakeShareable(new FMonitoredProcess(CmdExe, CommandLine, true));
 		UBTProcess->OnOutput().BindStatic(&FTargetPlatformManagerModule::OnStatusOutput);
 		SDKStatusMessage = TEXT("");
@@ -798,16 +804,15 @@ protected:
 			FPlatformProcess::Sleep(0.01f);
 		}
 
-		TArray<FString> LineArray;
-		SDKStatusMessage.ParseIntoArrayLines(&LineArray);
-		for (int Index = 0; Index < LineArray.Num(); ++Index)
+		TArray<FString> PlatArray;
+		SDKStatusMessage.ParseIntoArrayWS(&PlatArray);
+		for (int Index = 0; Index < PlatArray.Num()-2; ++Index)
 		{
-			TArray<FString> PlatArray;
-			LineArray[Index].ParseIntoArrayWS(&PlatArray);
-			if (PlatArray.Num() == 3)
+            FString Item = PlatArray[Index];
+			if (PlatArray[Index].Contains(TEXT("##PlatformValidate:")))
 			{
-				PlatformInfo::EPlatformSDKStatus Status = PlatArray[2] == TEXT("VALID") ? PlatformInfo::EPlatformSDKStatus::Installed : PlatformInfo::EPlatformSDKStatus::NotInstalled;
-				FString PlatformName = PlatArray[1];
+				PlatformInfo::EPlatformSDKStatus Status = PlatArray[Index+2].Contains(TEXT("INVALID")) ? PlatformInfo::EPlatformSDKStatus::NotInstalled : PlatformInfo::EPlatformSDKStatus::Installed;
+				FString PlatformName = PlatArray[Index+1];
 				if (PlatformName == TEXT("Win32") || PlatformName == TEXT("Win64"))
 				{
 					PlatformName = TEXT("Windows");
