fix for target dlls not being built as part of a github build when the SDK is not installed
#ue4
#platform

diff --git a/Engine/Source/Developer/TargetPlatform/TargetPlatform.Build.cs b/Engine/Source/Developer/TargetPlatform/TargetPlatform.Build.cs
--- a/Engine/Source/Developer/TargetPlatform/TargetPlatform.Build.cs
+++ b/Engine/Source/Developer/TargetPlatform/TargetPlatform.Build.cs
@@ -39,6 +39,21 @@ public class TargetPlatform : ModuleRules
 					DynamicallyLoadedModuleNames.Add("AudioFormatOgg");
 					DynamicallyLoadedModuleNames.Add("AudioFormatOpus"); 
 				}
+
+				if (Target.Type == TargetRules.TargetType.Editor || Target.Type == TargetRules.TargetType.Program)
+				{
+					DynamicallyLoadedModuleNames.Add("AndroidTargetPlatform");
+					DynamicallyLoadedModuleNames.Add("Android_PVRTCTargetPlatform");
+					DynamicallyLoadedModuleNames.Add("Android_ATCTargetPlatform");
+					DynamicallyLoadedModuleNames.Add("Android_DXTTargetPlatform");
+					DynamicallyLoadedModuleNames.Add("Android_ETC1TargetPlatform");
+					DynamicallyLoadedModuleNames.Add("Android_ETC2TargetPlatform");
+					DynamicallyLoadedModuleNames.Add("IOSTargetPlatform");
+					if (!UnrealBuildTool.UnrealBuildTool.RunningRocket())
+					{
+						DynamicallyLoadedModuleNames.Add("HTML5TargetPlatform");
+					}
+				}
 			}
 			else if (Target.Platform == UnrealTargetPlatform.Mac)
 			{
@@ -54,6 +69,21 @@ public class TargetPlatform : ModuleRules
 				{
 					DynamicallyLoadedModuleNames.Add("AudioFormatOgg");
 				}
+
+				if (Target.Type == TargetRules.TargetType.Editor || Target.Type == TargetRules.TargetType.Program)
+				{
+					DynamicallyLoadedModuleNames.Add("IOSTargetPlatform");
+					if (!UnrealBuildTool.UnrealBuildTool.RunningRocket())
+					{
+						DynamicallyLoadedModuleNames.Add("AndroidTargetPlatform");
+						DynamicallyLoadedModuleNames.Add("Android_PVRTCTargetPlatform");
+						DynamicallyLoadedModuleNames.Add("Android_ATCTargetPlatform");
+						DynamicallyLoadedModuleNames.Add("Android_DXTTargetPlatform");
+						DynamicallyLoadedModuleNames.Add("Android_ETC1TargetPlatform");
+						DynamicallyLoadedModuleNames.Add("Android_ETC2TargetPlatform");
+						DynamicallyLoadedModuleNames.Add("HTML5TargetPlatform");
+					}
+				}
 			}
             else if (Target.Platform == UnrealTargetPlatform.Linux)
             {
