another fix for the target platforms being added when they shouldn't
#ue4

diff --git a/Engine/Source/Developer/TargetPlatform/TargetPlatform.Build.cs b/Engine/Source/Developer/TargetPlatform/TargetPlatform.Build.cs
--- a/Engine/Source/Developer/TargetPlatform/TargetPlatform.Build.cs
+++ b/Engine/Source/Developer/TargetPlatform/TargetPlatform.Build.cs
@@ -49,7 +49,7 @@ public class TargetPlatform : ModuleRules
 					DynamicallyLoadedModuleNames.Add("Android_ETC1TargetPlatform");
 					DynamicallyLoadedModuleNames.Add("Android_ETC2TargetPlatform");
 					DynamicallyLoadedModuleNames.Add("IOSTargetPlatform");
-					if (!UnrealBuildTool.UnrealBuildTool.RunningRocket())
+					if (!UnrealBuildTool.UnrealBuildTool.RunningRocket() && !UnrealBuildTool.UnrealBuildTool.BuildingRocket())
 					{
 						DynamicallyLoadedModuleNames.Add("HTML5TargetPlatform");
 					}
@@ -73,7 +73,7 @@ public class TargetPlatform : ModuleRules
 				if (Target.Type == TargetRules.TargetType.Editor || Target.Type == TargetRules.TargetType.Program)
 				{
 					DynamicallyLoadedModuleNames.Add("IOSTargetPlatform");
-					if (!UnrealBuildTool.UnrealBuildTool.RunningRocket())
+					if (!UnrealBuildTool.UnrealBuildTool.RunningRocket() && !UnrealBuildTool.UnrealBuildTool.BuildingRocket())
 					{
 						DynamicallyLoadedModuleNames.Add("AndroidTargetPlatform");
 						DynamicallyLoadedModuleNames.Add("Android_PVRTCTargetPlatform");
diff --git a/Engine/Source/Programs/UnrealBuildTool/HTML5/UEBuildHTML5.cs b/Engine/Source/Programs/UnrealBuildTool/HTML5/UEBuildHTML5.cs
--- a/Engine/Source/Programs/UnrealBuildTool/HTML5/UEBuildHTML5.cs
+++ b/Engine/Source/Programs/UnrealBuildTool/HTML5/UEBuildHTML5.cs
@@ -336,12 +336,6 @@ namespace UnrealBuildTool
          */
         public override void ModifyNewlyLoadedModule(UEBuildModule InModule, TargetInfo Target)
         {
-			[REDACTED_EMAIL]: Add platform support
-			if (UnrealBuildTool.RunningRocket() || UnrealBuildTool.BuildingRocket())
-			{
-				return;
-			}
-
 			if (Target.Platform == UnrealTargetPlatform.HTML5)
             {
                 if (InModule.ToString() == "Core")
