Don't run AutoSDKs code if AutoSDKs aren't enabled.

diff --git a/Engine/Source/Developer/TargetPlatform/Private/TargetPlatformManagerModule.cpp b/Engine/Source/Developer/TargetPlatform/Private/TargetPlatformManagerModule.cpp
--- a/Engine/Source/Developer/TargetPlatform/Private/TargetPlatformManagerModule.cpp
+++ b/Engine/Source/Developer/TargetPlatform/Private/TargetPlatformManagerModule.cpp
@@ -32,25 +32,34 @@ public:
 	{
 
 #if AUTOSDKS_ENABLED		
-		// amortize UBT cost by calling it once for all platforms, rather than once per platform.
-		FString UBTParams(TEXT("-autosdkonly"));
-		int32 UBTReturnCode = -1;
-		FString UBTOutput;
-		if (!FUBTInvoker::InvokeUnrealBuildToolSync(UBTParams, *GLog, true, UBTReturnCode, UBTOutput))
-		{
-			UE_LOG(LogTargetPlatformManager, Fatal, TEXT("Failed to run UBT to check SDK status!"));
-		}
+		static const FString SDKRootEnvFar(TEXT("UE_SDKS_ROOT"));		
+		const int32 MaxPathSize = 32768;
+		TCHAR SDKPath[MaxPathSize] = { 0 };
+		FPlatformMisc::GetEnvironmentVariable(*SDKRootEnvFar, SDKPath, MaxPathSize);
 
-		// we have to setup our local environment according to AutoSDKs or the ITargetPlatform's IsSDkInstalled calls may fail
-		// before we get a change to setup for a given platform.  Use the platforminfo list to avoid any kind of interdependency.
-		int32 NumPlatforms;
-		const PlatformInfo::FPlatformInfo* PlatformInfoArray = PlatformInfo::GetPlatformInfoArray(NumPlatforms);
-		for (int32 i = 0; i < NumPlatforms; ++i)
-		{
-			const PlatformInfo::FPlatformInfo& PlatformInfo = PlatformInfoArray[i];
-			if (PlatformInfo.AutoSDKPath.Len() > 0)
+		// AutoSDKs only enabled if UE_SDKS_ROOT is set.
+		if (SDKPath[0] != 0)
+		{					
+			// amortize UBT cost by calling it once for all platforms, rather than once per platform.
+			FString UBTParams(TEXT("-autosdkonly"));
+			int32 UBTReturnCode = -1;
+			FString UBTOutput;
+			if (!FUBTInvoker::InvokeUnrealBuildToolSync(UBTParams, *GLog, true, UBTReturnCode, UBTOutput))
+			{
+				UE_LOG(LogTargetPlatformManager, Fatal, TEXT("Failed to run UBT to check SDK status!"));
+			}
+
+			// we have to setup our local environment according to AutoSDKs or the ITargetPlatform's IsSDkInstalled calls may fail
+			// before we get a change to setup for a given platform.  Use the platforminfo list to avoid any kind of interdependency.
+			int32 NumPlatforms;
+			const PlatformInfo::FPlatformInfo* PlatformInfoArray = PlatformInfo::GetPlatformInfoArray(NumPlatforms);
+			for (int32 i = 0; i < NumPlatforms; ++i)
 			{
-				SetupAndValidateAutoSDK(PlatformInfo.AutoSDKPath);
+				const PlatformInfo::FPlatformInfo& PlatformInfo = PlatformInfoArray[i];
+				if (PlatformInfo.AutoSDKPath.Len() > 0)
+				{
+					SetupAndValidateAutoSDK(PlatformInfo.AutoSDKPath);
+				}
 			}
 		}
 #endif
