Nicer AutoSDKs checks

diff --git a/Engine/Source/Developer/TargetPlatform/Private/TargetPlatformManagerModule.cpp b/Engine/Source/Developer/TargetPlatform/Private/TargetPlatformManagerModule.cpp
--- a/Engine/Source/Developer/TargetPlatform/Private/TargetPlatformManagerModule.cpp
+++ b/Engine/Source/Developer/TargetPlatform/Private/TargetPlatformManagerModule.cpp
@@ -32,13 +32,9 @@ public:
 	{
 
 #if AUTOSDKS_ENABLED		
-		static const FString SDKRootEnvFar(TEXT("UE_SDKS_ROOT"));		
-		const int32 MaxPathSize = 32768;
-		TCHAR SDKPath[MaxPathSize] = { 0 };
-		FPlatformMisc::GetEnvironmentVariable(*SDKRootEnvFar, SDKPath, MaxPathSize);
-
+		
 		// AutoSDKs only enabled if UE_SDKS_ROOT is set.
-		if (SDKPath[0] != 0)
+		if (IsAutoSDKsEnabled())
 		{					
 			// amortize UBT cost by calling it once for all platforms, rather than once per platform.
 			FString UBTParams(TEXT("-autosdkonly"));
@@ -525,6 +521,21 @@ public:
 
 protected:
 
+	bool IsAutoSDKsEnabled()
+	{
+		static const FString SDKRootEnvFar(TEXT("UE_SDKS_ROOT"));
+		const int32 MaxPathSize = 16384;
+		TCHAR SDKPath[MaxPathSize] = { 0 };
+		FPlatformMisc::GetEnvironmentVariable(*SDKRootEnvFar, SDKPath, MaxPathSize);
+
+		// AutoSDKs only enabled if UE_SDKS_ROOT is set.
+		if (SDKPath[0] != 0)
+		{
+			return true;
+		}
+		return false;
+	}
+
 	/**
 	 * Discovers the available target platforms.
 	 */
@@ -597,16 +608,11 @@ protected:
 	bool SetupEnvironmentFromAutoSDK(const FString& AutoSDKPath)
 	{						
 #if AUTOSDKS_ENABLED
-		static const FString SDKRootEnvFar(TEXT("UE_SDKS_ROOT"));		
-		const int32 MaxPathSize = 32768;
-		TCHAR SDKPath[MaxPathSize] = { 0 };
-		FPlatformMisc::GetEnvironmentVariable(*SDKRootEnvFar, SDKPath, MaxPathSize);
-
-		// AutoSDKs only enabled if UE_SDKS_ROOT is set.
-		if (SDKPath[0] == 0)
+		
+		if (!IsAutoSDKsEnabled())
 		{
 			return true;
-		}		
+		}
 
 		// Invoke UBT to perform SDK switching, or detect that a proper manual SDK is already setup.				
 #if PLATFORM_WINDOWS
@@ -615,6 +621,11 @@ protected:
 #error Fill in your host platform directory
 #endif		
 
+		static const FString SDKRootEnvFar(TEXT("UE_SDKS_ROOT"));
+		const int32 MaxPathSize = 16384;
+		TCHAR SDKPath[MaxPathSize] = { 0 };
+		FPlatformMisc::GetEnvironmentVariable(*SDKRootEnvFar, SDKPath, MaxPathSize);
+
 		FString TargetSDKRoot = FPaths::Combine(SDKPath, *HostPlatform, *AutoSDKPath);
 		static const FString SDKInstallManifestFileName(TEXT("CurrentlyInstalled.txt"));
 		FString SDKInstallManifestFilePath = FPaths::Combine(*TargetSDKRoot, *SDKInstallManifestFileName);
