log available platforms in case it wasn't recognized

diff --git a/Engine/Source/Developer/TargetPlatform/Private/TargetPlatformManagerModule.cpp b/Engine/Source/Developer/TargetPlatform/Private/TargetPlatformManagerModule.cpp
--- a/Engine/Source/Developer/TargetPlatform/Private/TargetPlatformManagerModule.cpp
+++ b/Engine/Source/Developer/TargetPlatform/Private/TargetPlatformManagerModule.cpp
@@ -188,19 +188,28 @@ public:
 
 					PlatformStr.ParseIntoArray(&PlatformNames, TEXT("+"), true);
 
+					// for nicer user response
+					FString AvailablePlatforms;
+
 					for (int32 Index = 0; Index < TargetPlatforms.Num(); Index++)
 					{
 						if (PlatformNames.Contains(TargetPlatforms[Index]->PlatformName()))
 						{							
 							Results.Add(TargetPlatforms[Index]);						
 						}
+
+						if(!AvailablePlatforms.IsEmpty())
+						{
+							AvailablePlatforms += TEXT(", ");
+						}
+						AvailablePlatforms += TargetPlatforms[Index]->PlatformName();
 					}
 
 					if (Results.Num() == 0)
 					{
 						// An invalid platform was specified...
 						// Inform the user and exit.
-						UE_LOG(LogTargetPlatformManager, Fatal, TEXT("Invalid target platform specified (%s)."), *PlatformStr);
+						UE_LOG(LogTargetPlatformManager, Fatal, TEXT("Invalid target platform specified (%s). Available = { %s } "), *PlatformStr, *AvailablePlatforms);
 					}
 				}
 			}
