TargetPlatform: Replaced namespaced enums with enum classes

diff --git a/Engine/Source/Developer/Android/AndroidTargetPlatform/Private/AndroidTargetDevice.h b/Engine/Source/Developer/Android/AndroidTargetPlatform/Private/AndroidTargetDevice.h
--- a/Engine/Source/Developer/Android/AndroidTargetPlatform/Private/AndroidTargetDevice.h
+++ b/Engine/Source/Developer/Android/AndroidTargetPlatform/Private/AndroidTargetDevice.h
@@ -111,7 +111,7 @@ public:
 	{
 	}
 
-	virtual ETargetDeviceTypes::Type GetDeviceType() const override
+	virtual ETargetDeviceTypes GetDeviceType() const override
 	{
 		//@TODO: How to distinguish between a Tablet and a Phone (or a TV microconsole, etc...), and is it important?
 		return ETargetDeviceTypes::Tablet;
@@ -163,7 +163,7 @@ public:
 	virtual bool Launch(const FString& AppId, EBuildConfigurations::Type BuildConfiguration, EBuildTargets::Type BuildTarget, const FString& Params, uint32* OutProcessId);
 	virtual bool Reboot(bool bReconnect = false) override;
 	virtual bool Run(const FString& ExecutablePath, const FString& Params, uint32* OutProcessId) override;
-	virtual bool SupportsFeature(ETargetDeviceFeatures::Type Feature) const override;
+	virtual bool SupportsFeature(ETargetDeviceFeatures Feature) const override;
 	virtual bool SupportsSdkVersion(const FString& VersionString) const override;
 	virtual bool TerminateProcess(const int32 ProcessId) override;
 	virtual void SetUserCredentials(const FString& UserName, const FString& UserPassword) override;
diff --git a/Engine/Source/Developer/Android/AndroidTargetPlatform/Private/AndroidTargetDevice.inl b/Engine/Source/Developer/Android/AndroidTargetPlatform/Private/AndroidTargetDevice.inl
--- a/Engine/Source/Developer/Android/AndroidTargetPlatform/Private/AndroidTargetDevice.inl
+++ b/Engine/Source/Developer/Android/AndroidTargetPlatform/Private/AndroidTargetDevice.inl
@@ -73,7 +73,7 @@ inline bool FAndroidTargetDevice::Run( const FString& ExecutablePath, const FStr
 }
 
 
-inline bool FAndroidTargetDevice::SupportsFeature( ETargetDeviceFeatures::Type Feature ) const
+inline bool FAndroidTargetDevice::SupportsFeature( ETargetDeviceFeatures Feature ) const
 {
 	switch (Feature)
 	{
diff --git a/Engine/Source/Developer/Android/AndroidTargetPlatform/Private/AndroidTargetPlatform.h b/Engine/Source/Developer/Android/AndroidTargetPlatform/Private/AndroidTargetPlatform.h
--- a/Engine/Source/Developer/Android/AndroidTargetPlatform/Private/AndroidTargetPlatform.h
+++ b/Engine/Source/Developer/Android/AndroidTargetPlatform/Private/AndroidTargetPlatform.h
@@ -106,7 +106,7 @@ public:
 
 	virtual bool IsSdkInstalled(bool bProjectHasCode, FString& OutDocumentationPath) const override;
 
-	virtual bool SupportsFeature( ETargetPlatformFeatures::Type Feature ) const override;
+	virtual bool SupportsFeature( ETargetPlatformFeatures Feature ) const override;
 
 	virtual bool SupportsTextureFormat( FName Format ) const 
 	{
diff --git a/Engine/Source/Developer/Android/AndroidTargetPlatform/Private/AndroidTargetPlatform.inl b/Engine/Source/Developer/Android/AndroidTargetPlatform/Private/AndroidTargetPlatform.inl
--- a/Engine/Source/Developer/Android/AndroidTargetPlatform/Private/AndroidTargetPlatform.inl
+++ b/Engine/Source/Developer/Android/AndroidTargetPlatform/Private/AndroidTargetPlatform.inl
@@ -139,7 +139,7 @@ inline bool FAndroidTargetPlatform<TPlatformProperties>::IsSdkInstalled(bool bPr
 
 
 template<class TPlatformProperties>
-inline bool FAndroidTargetPlatform<TPlatformProperties>::SupportsFeature( ETargetPlatformFeatures::Type Feature ) const
+inline bool FAndroidTargetPlatform<TPlatformProperties>::SupportsFeature( ETargetPlatformFeatures Feature ) const
 {
 	if (Feature == ETargetPlatformFeatures::Packaging)
 	{
diff --git a/Engine/Source/Developer/Android/Android_ES31TargetPlatform/Private/Android_ES31TargetPlatformModule.cpp b/Engine/Source/Developer/Android/Android_ES31TargetPlatform/Private/Android_ES31TargetPlatformModule.cpp
--- a/Engine/Source/Developer/Android/Android_ES31TargetPlatform/Private/Android_ES31TargetPlatformModule.cpp
+++ b/Engine/Source/Developer/Android/Android_ES31TargetPlatform/Private/Android_ES31TargetPlatformModule.cpp
@@ -47,7 +47,7 @@ class FAndroid_ES31TargetPlatform
 		return (ExtensionsString.Contains(TEXT("GL_EXT_tessellation_shader")));
 	}
 
-	virtual bool SupportsFeature( ETargetPlatformFeatures::Type Feature ) const override
+	virtual bool SupportsFeature( ETargetPlatformFeatures Feature ) const override
 	{
 		switch (Feature)
 		{
diff --git a/Engine/Source/Developer/HTML5/HTML5TargetPlatform/Private/HTML5TargetDevice.cpp b/Engine/Source/Developer/HTML5/HTML5TargetPlatform/Private/HTML5TargetDevice.cpp
--- a/Engine/Source/Developer/HTML5/HTML5TargetPlatform/Private/HTML5TargetDevice.cpp
+++ b/Engine/Source/Developer/HTML5/HTML5TargetPlatform/Private/HTML5TargetDevice.cpp
@@ -77,7 +77,7 @@ bool FHTML5TargetDevice::Run( const FString& ExecutablePath, const FString& Para
 }
 
 
-bool FHTML5TargetDevice::SupportsFeature( ETargetDeviceFeatures::Type Feature ) const
+bool FHTML5TargetDevice::SupportsFeature( ETargetDeviceFeatures Feature ) const
 {
 	switch (Feature)
 	{
diff --git a/Engine/Source/Developer/HTML5/HTML5TargetPlatform/Private/HTML5TargetDevice.h b/Engine/Source/Developer/HTML5/HTML5TargetPlatform/Private/HTML5TargetDevice.h
--- a/Engine/Source/Developer/HTML5/HTML5TargetPlatform/Private/HTML5TargetDevice.h
+++ b/Engine/Source/Developer/HTML5/HTML5TargetPlatform/Private/HTML5TargetDevice.h
@@ -58,7 +58,7 @@ public:
 
 	virtual void Disconnect( ) override;
 
-	virtual ETargetDeviceTypes::Type GetDeviceType( ) const override
+	virtual ETargetDeviceTypes GetDeviceType( ) const override
 	{
 		return ETargetDeviceTypes::Browser;
 	}
@@ -96,7 +96,7 @@ public:
 
 	virtual bool Run( const FString& ExecutablePath, const FString& Params, uint32* OutProcessId ) override;
 
-	virtual bool SupportsFeature( ETargetDeviceFeatures::Type Feature ) const override;
+	virtual bool SupportsFeature( ETargetDeviceFeatures Feature ) const override;
 
 	virtual bool SupportsSdkVersion( const FString& VersionString ) const override;
 
diff --git a/Engine/Source/Developer/HTML5/HTML5TargetPlatform/Private/HTML5TargetPlatform.h b/Engine/Source/Developer/HTML5/HTML5TargetPlatform/Private/HTML5TargetPlatform.h
--- a/Engine/Source/Developer/HTML5/HTML5TargetPlatform/Private/HTML5TargetPlatform.h
+++ b/Engine/Source/Developer/HTML5/HTML5TargetPlatform/Private/HTML5TargetPlatform.h
@@ -44,7 +44,7 @@ public:
 
 	virtual bool IsRunningPlatform( ) const override;
 
-	virtual bool SupportsFeature( ETargetPlatformFeatures::Type Feature ) const override
+	virtual bool SupportsFeature( ETargetPlatformFeatures Feature ) const override
 	{
 		if (Feature == ETargetPlatformFeatures::Packaging)
 		{
diff --git a/Engine/Source/Developer/Linux/LinuxNoEditorTargetPlatform/Private/LinuxTargetDevice.cpp b/Engine/Source/Developer/Linux/LinuxNoEditorTargetPlatform/Private/LinuxTargetDevice.cpp
--- a/Engine/Source/Developer/Linux/LinuxNoEditorTargetPlatform/Private/LinuxTargetDevice.cpp
+++ b/Engine/Source/Developer/Linux/LinuxNoEditorTargetPlatform/Private/LinuxTargetDevice.cpp
@@ -38,7 +38,7 @@ bool FLinuxTargetDevice::Run( const FString& ExecutablePath, const FString& Para
 	return false;
 }
 
-bool FLinuxTargetDevice::SupportsFeature( ETargetDeviceFeatures::Type Feature ) const
+bool FLinuxTargetDevice::SupportsFeature( ETargetDeviceFeatures Feature ) const
 {
 	return false;
 }
diff --git a/Engine/Source/Developer/Linux/LinuxServerTargetPlatform/Private/LinuxTargetDevice.cpp b/Engine/Source/Developer/Linux/LinuxServerTargetPlatform/Private/LinuxTargetDevice.cpp
--- a/Engine/Source/Developer/Linux/LinuxServerTargetPlatform/Private/LinuxTargetDevice.cpp
+++ b/Engine/Source/Developer/Linux/LinuxServerTargetPlatform/Private/LinuxTargetDevice.cpp
@@ -38,7 +38,7 @@ bool FLinuxTargetDevice::Run( const FString& ExecutablePath, const FString& Para
 	return false;
 }
 
-bool FLinuxTargetDevice::SupportsFeature( ETargetDeviceFeatures::Type Feature ) const
+bool FLinuxTargetDevice::SupportsFeature( ETargetDeviceFeatures Feature ) const
 {
 	return false;
 }
diff --git a/Engine/Source/Developer/Linux/LinuxTargetPlatform/Private/LinuxTargetDevice.cpp b/Engine/Source/Developer/Linux/LinuxTargetPlatform/Private/LinuxTargetDevice.cpp
--- a/Engine/Source/Developer/Linux/LinuxTargetPlatform/Private/LinuxTargetDevice.cpp
+++ b/Engine/Source/Developer/Linux/LinuxTargetPlatform/Private/LinuxTargetDevice.cpp
@@ -107,7 +107,7 @@ bool FLinuxTargetDevice::Run( const FString& ExecutablePath, const FString& Para
 #endif // PLATFORM_LINUX
 }
 
-bool FLinuxTargetDevice::SupportsFeature( ETargetDeviceFeatures::Type Feature ) const
+bool FLinuxTargetDevice::SupportsFeature( ETargetDeviceFeatures Feature ) const
 {
 	switch (Feature)
 	{
diff --git a/Engine/Source/Developer/Linux/LinuxTargetPlatform/Private/LinuxTargetDevice.h b/Engine/Source/Developer/Linux/LinuxTargetPlatform/Private/LinuxTargetDevice.h
--- a/Engine/Source/Developer/Linux/LinuxTargetPlatform/Private/LinuxTargetDevice.h
+++ b/Engine/Source/Developer/Linux/LinuxTargetPlatform/Private/LinuxTargetDevice.h
@@ -50,7 +50,7 @@ public:
 	virtual void Disconnect( )
 	{ }
 
-	virtual ETargetDeviceTypes::Type GetDeviceType( ) const override
+	virtual ETargetDeviceTypes GetDeviceType( ) const override
 	{
 		return ETargetDeviceTypes::Desktop;
 	}
@@ -103,7 +103,7 @@ public:
 
 	virtual bool Run( const FString& ExecutablePath, const FString& Params, uint32* OutProcessId ) override;
 
-	virtual bool SupportsFeature( ETargetDeviceFeatures::Type Feature ) const override;
+	virtual bool SupportsFeature( ETargetDeviceFeatures Feature ) const override;
 
 	virtual bool SupportsSdkVersion( const FString& VersionString ) const override;
 
diff --git a/Engine/Source/Developer/Linux/LinuxTargetPlatform/Private/LinuxTargetPlatform.h b/Engine/Source/Developer/Linux/LinuxTargetPlatform/Private/LinuxTargetPlatform.h
--- a/Engine/Source/Developer/Linux/LinuxTargetPlatform/Private/LinuxTargetPlatform.h
+++ b/Engine/Source/Developer/Linux/LinuxTargetPlatform/Private/LinuxTargetPlatform.h
@@ -121,7 +121,7 @@ public:
 		return PLATFORM_LINUX && !UE_SERVER && !UE_GAME && WITH_EDITOR && HAS_EDITOR_DATA;
 	}
 
-	virtual bool SupportsFeature(ETargetPlatformFeatures::Type Feature) const override
+	virtual bool SupportsFeature(ETargetPlatformFeatures Feature) const override
 	{
 		if (Feature == ETargetPlatformFeatures::UserCredentials || Feature == ETargetPlatformFeatures::Packaging)
 		{
diff --git a/Engine/Source/Developer/Mac/MacTargetPlatform/Private/GenericMacTargetPlatform.h b/Engine/Source/Developer/Mac/MacTargetPlatform/Private/GenericMacTargetPlatform.h
--- a/Engine/Source/Developer/Mac/MacTargetPlatform/Private/GenericMacTargetPlatform.h
+++ b/Engine/Source/Developer/Mac/MacTargetPlatform/Private/GenericMacTargetPlatform.h
@@ -81,7 +81,7 @@ public:
 		return PLATFORM_MAC && !UE_SERVER && !UE_GAME && WITH_EDITOR && HAS_EDITOR_DATA;
 	}
 
-	virtual bool SupportsFeature( ETargetPlatformFeatures::Type Feature ) const override
+	virtual bool SupportsFeature( ETargetPlatformFeatures Feature ) const override
 	{
 		// we currently do not have a build target for MacServer
 		if (Feature == ETargetPlatformFeatures::Packaging)
diff --git a/Engine/Source/Developer/Mac/MacTargetPlatform/Private/LocalMacTargetDevice.h b/Engine/Source/Developer/Mac/MacTargetPlatform/Private/LocalMacTargetDevice.h
--- a/Engine/Source/Developer/Mac/MacTargetPlatform/Private/LocalMacTargetDevice.h
+++ b/Engine/Source/Developer/Mac/MacTargetPlatform/Private/LocalMacTargetDevice.h
@@ -61,7 +61,7 @@ public:
 	virtual void Disconnect( )
 	{ }
 
-	virtual ETargetDeviceTypes::Type GetDeviceType( ) const override
+	virtual ETargetDeviceTypes GetDeviceType( ) const override
 	{
 		return ETargetDeviceTypes::Desktop;
 	}
@@ -145,7 +145,7 @@ public:
 		return ProcessHandle.Close();
 	}
 
-	virtual bool SupportsFeature( ETargetDeviceFeatures::Type Feature ) const override
+	virtual bool SupportsFeature( ETargetDeviceFeatures Feature ) const override
 	{
 		switch (Feature)
 		{
diff --git a/Engine/Source/Developer/TargetDeviceServices/Private/TargetDeviceService.cpp b/Engine/Source/Developer/TargetDeviceServices/Private/TargetDeviceService.cpp
--- a/Engine/Source/Developer/TargetDeviceServices/Private/TargetDeviceService.cpp
+++ b/Engine/Source/Developer/TargetDeviceServices/Private/TargetDeviceService.cpp
@@ -395,7 +395,7 @@ void FTargetDeviceService::HandlePingMessage(const FTargetDeviceServicePing& InM
 		FTargetDeviceServicePong* Message = new FTargetDeviceServicePong();
 
 		Message->Name = DefaultDevice->GetName();
-		Message->Type = ETargetDeviceTypes::ToString(DefaultDevice->GetDeviceType());
+		Message->Type = TargetDeviceTypes::ToString(DefaultDevice->GetDeviceType());
 		Message->HostName = FPlatformProcess::ComputerName();
 		Message->HostUser = FPlatformProcess::UserName(true);
 		Message->Connected = DefaultDevice->IsConnected();
diff --git a/Engine/Source/Developer/TargetPlatform/Public/Common/TargetPlatformBase.h b/Engine/Source/Developer/TargetPlatform/Public/Common/TargetPlatformBase.h
--- a/Engine/Source/Developer/TargetPlatform/Public/Common/TargetPlatformBase.h
+++ b/Engine/Source/Developer/TargetPlatform/Public/Common/TargetPlatformBase.h
@@ -168,10 +168,10 @@ public:
 
 	virtual int DoesntHaveRequirements(const FString& ProjectPath, bool bProjectHasCode, FString& OutDocumentationPath) const override
 	{
-		int bReadyToBuild = ETargetPlatformReadyStatus::Ready; // @todo How do we check that the iOS SDK is installed when building from Windows? Is that even possible?
+		int bReadyToBuild = static_cast<int>(ETargetPlatformReadyStatus::Ready); // @todo How do we check that the iOS SDK is installed when building from Windows? Is that even possible?
 		if (!IsSdkInstalled(bProjectHasCode, OutDocumentationPath))
 		{
-			bReadyToBuild |= ETargetPlatformReadyStatus::SDKNotFound;
+			bReadyToBuild |= static_cast<int>(ETargetPlatformReadyStatus::SDKNotFound);
 		}
 		return bReadyToBuild;
 	}
@@ -283,7 +283,7 @@ public:
 		return TPlatformProperties::SupportsAutoSDK();
 	}
 
-	virtual bool SupportsFeature( ETargetPlatformFeatures::Type Feature ) const override
+	virtual bool SupportsFeature( ETargetPlatformFeatures Feature ) const override
 	{
 		switch (Feature)
 		{
diff --git a/Engine/Source/Developer/TargetPlatform/Public/Interfaces/ITargetDevice.h b/Engine/Source/Developer/TargetPlatform/Public/Interfaces/ITargetDevice.h
--- a/Engine/Source/Developer/TargetPlatform/Public/Interfaces/ITargetDevice.h
+++ b/Engine/Source/Developer/TargetPlatform/Public/Interfaces/ITargetDevice.h
@@ -3,141 +3,135 @@
 #pragma once
 
 
-namespace ETargetDeviceFeatures
+/**
+ * Enumerates features that may be supported by target devices.
+ */
+enum class ETargetDeviceFeatures
 {
-	/**
-	 * Enumerates features that may be supported by target devices.
-	 */
-	enum Type
-	{
-		/** Multiple instances of a game can run at the same time. */
-		MultiLaunch,
+	/** Multiple instances of a game can run at the same time. */
+	MultiLaunch,
 
-		/** The device can be powered off remotely. */
-		PowerOff,
+	/** The device can be powered off remotely. */
+	PowerOff,
 
-		/** The device can be powered on remotely. */	  
-		PowerOn,
+	/** The device can be powered on remotely. */	  
+	PowerOn,
 
-		/** Snapshot of processes running on the device. */
-		ProcessSnapshot,
+	/** Snapshot of processes running on the device. */
+	ProcessSnapshot,
 
-		/** The device can be rebooted remotely. */
-		Reboot
-	};
-}
+	/** The device can be rebooted remotely. */
+	Reboot
+};
 
 
-namespace ETargetDeviceTypes
+/**
+ * Enumerates target device types.
+ */
+enum class ETargetDeviceTypes
 {
-	/**
-	 * Enumerates target device types.
-	 */
-	enum Type
-	{
-		/** Indeterminate device type. */
-		Indeterminate,
+	/** Indeterminate device type. */
+	Indeterminate,
 
-		/** The device is a web browser (i.e. Flash). */
-		Browser,
+	/** The device is a web browser (i.e. Flash). */
+	Browser,
 
-		/** The device is a game console. */
-		Console,
+	/** The device is a game console. */
+	Console,
 
-		/** The device is a desktop computer. */
-		Desktop,
+	/** The device is a desktop computer. */
+	Desktop,
 
-		/** The device is a smart phone. */
-		Phone,
+	/** The device is a smart phone. */
+	Phone,
 
-		/** The device is a tablet computer. */
-		Tablet
-	};
+	/** The device is a tablet computer. */
+	Tablet
+};
 
 
+namespace TargetDeviceTypes
+{
 	/**
 	 * Returns the string representation of the specified ETargetDeviceTypes value.
 	 *
 	 * @param Configuration The value to get the string for.
 	 * @return A string value.
 	 */
-	inline FString ToString(ETargetDeviceTypes::Type DeviceType)
+	inline FString ToString(ETargetDeviceTypes DeviceType)
 	{
 		switch (DeviceType)
 		{
-			case Browser:
-				return FString("Browser");
+		case ETargetDeviceTypes::Browser:
+			return FString("Browser");
 
-			case Console:
-				return FString("Console");
+		case ETargetDeviceTypes::Console:
+			return FString("Console");
 
-			case Desktop:
-				return FString("Desktop");
+		case ETargetDeviceTypes::Desktop:
+			return FString("Desktop");
 
-			case Phone:
-				return FString("Phone");
+		case ETargetDeviceTypes::Phone:
+			return FString("Phone");
 
-			case Tablet:
-				return FString("Tablet");
+		case ETargetDeviceTypes::Tablet:
+			return FString("Tablet");
 
-			default:
-				return FString("Indeterminate");
+		default:
+			return FString("Indeterminate");
 		}
 	}
 }
 
 
-namespace ETargetDeviceThreadStates
+/**
+ * Enumerates states of threads running on a target device.
+ */
+enum class ETargetDeviceThreadStates
 {
-	/**
-	 * Enumerates thread states.
-	 */
-	enum Type
-	{
-		/** Unknown thread state. */
-		Unknown,
+	/** Unknown thread state. */
+	Unknown,
 
-		/** The thread can run, but is not running right now. */
-		CanRun,
+	/** The thread can run, but is not running right now. */
+	CanRun,
 
-		/** The thread is inactive, i.e. has just been created or exited. */
-		Inactive,
+	/** The thread is inactive, i.e. has just been created or exited. */
+	Inactive,
 
-		/** The thread cannot run right now. */
-		Inhibited,
+	/** The thread cannot run right now. */
+	Inhibited,
 
-		/** The thread is in the run queue. */
-		RunQueue,
+	/** The thread is in the run queue. */
+	RunQueue,
 
-		/** The thread is running. */
-		Running
-	};
-}
+	/** The thread is running. */
+	Running
+};
 
 
-namespace ETargetDeviceThreadWaitStates
+/**
+ * Enumerates wait states of threads running on a target device.
+ */
+enum class ETargetDeviceThreadWaitStates
 {
-	enum Type
-	{
-		/** Unknown wait state. */
-		Unknown,
+	/** Unknown wait state. */
+	Unknown,
 
-		/** The thread is blocked by a lock. */
-		Locked,
+	/** The thread is blocked by a lock. */
+	Locked,
 	
-		/** The thread is sleeping. */
-		Sleeping,
+	/** The thread is sleeping. */
+	Sleeping,
 
-		/** The thread is suspended. */
-		Suspended,
+	/** The thread is suspended. */
+	Suspended,
 
-		/** The thread is swapped. */
-		Swapped,
+	/** The thread is swapped. */
+	Swapped,
 
-		/** The thread is waiting on an interrupt. */
-		Waiting
-	};
-}
+	/** The thread is waiting on an interrupt. */
+	Waiting
+};
 
 
 /**
@@ -158,10 +152,10 @@ struct FTargetDeviceThreadInfo
 	uint64 StackSize;
 
 	/** Holds the thread's current state. */
-	ETargetDeviceThreadStates::Type State;
+	ETargetDeviceThreadStates State;
 
 	/** Holds the thread's current wait state. */
-	ETargetDeviceThreadWaitStates::Type WaitState;
+	ETargetDeviceThreadWaitStates WaitState;
 };
 
 
@@ -230,7 +224,7 @@ public:
 	 *
 	 * @return Device type.
 	 */
-	virtual ETargetDeviceTypes::Type GetDeviceType( ) const = 0;
+	virtual ETargetDeviceTypes GetDeviceType( ) const = 0;
 
 	/**
 	 * Gets the unique device identifier.
@@ -340,7 +334,7 @@ public:
 	 * @param Feature The feature to check.
 	 * @return true if the feature is supported, false otherwise.
 	 */
-	virtual bool SupportsFeature( ETargetDeviceFeatures::Type Feature ) const = 0;
+	virtual bool SupportsFeature( ETargetDeviceFeatures Feature ) const = 0;
 
 	/**
 	 * Checks whether this device supports the specified SDK version.
diff --git a/Engine/Source/Developer/TargetPlatform/Public/Interfaces/ITargetPlatform.h b/Engine/Source/Developer/TargetPlatform/Public/Interfaces/ITargetPlatform.h
--- a/Engine/Source/Developer/TargetPlatform/Public/Interfaces/ITargetPlatform.h
+++ b/Engine/Source/Developer/TargetPlatform/Public/Interfaces/ITargetPlatform.h
@@ -10,112 +10,82 @@ namespace PlatformInfo
 }
 
 
-namespace ETargetPlatformBuildArtifacts
+/**
+ * Enumerates features that may be supported by target platforms.
+ */
+enum class ETargetPlatformFeatures
 {
-	/**
-	 * Enumerates build artifact types.
-	 */
-	enum Type
-	{
-		/** Include content files. */
-		Content = 0x1,
-
-		/**
-		 * Include debug symbol files, i.e PDB.
-		 *
-		 * Requires 'Engine' artifact type.
-		 */
-		DebugSymbols = 0x2,
-
-		/** Include Engine binaries and DLLs. */
-		Engine = 0x4,
-
-		/** Include tools. */
-		Tools = 0x8,
-	};
-}
+	/** Audio Streaming */
+	AudioStreaming,
 
+	/** Distance field shadows. */
+	DistanceFieldShadows,
 
-namespace ETargetPlatformFeatures
-{
-	/**
-	 * Enumerates features that may be supported by target platforms.
-	 */
-	enum Type
-	{
-		/** Audio Streaming */
-		AudioStreaming,
+	/** Gray scale SRGB texture formats support. */
+	GrayscaleSRGB,
 
-		/** Distance field shadows. */
-		DistanceFieldShadows,
+	/** High quality light maps. */
+	HighQualityLightmaps,
 
-		/** Gray scale SRGB texture formats support. */
-		GrayscaleSRGB,
+	/** Low quality light maps. */
+	LowQualityLightmaps,
 
-		/** High quality light maps. */
-		HighQualityLightmaps,
+	/** Run multiple game instances on a single device. */
+	MultipleGameInstances,
 
-		/** Low quality light maps. */
-		LowQualityLightmaps,
+	/** Builds can be packaged for this platform. */
+	Packaging,
 
-		/** Run multiple game instances on a single device. */
-		MultipleGameInstances,
+	/** Connect and disconnect devices through the SDK. */
+	SdkConnectDisconnect,
 
-		/** Builds can be packaged for this platform. */
-		Packaging,
+	/** GPU tesselation. */
+	Tessellation,
 
-		/** Connect and disconnect devices through the SDK. */
-		SdkConnectDisconnect,
+	/** Texture streaming. */
+	TextureStreaming,
 
-		/** GPU tesselation. */
-		Tessellation,
+	/** User credentials are required to use the device. */
+	UserCredentials,
 
-		/** Texture streaming. */
-		TextureStreaming,
-
-		/** User credentials are required to use the device. */
-		UserCredentials,
-
-		/** Vertex Shader Texture Sampling. */
-		VertexShaderTextureSampling,
-	};
+	/** Vertex Shader Texture Sampling. */
+	VertexShaderTextureSampling,
 };
 
 
+/**
+ * Flags specifying what is needed to be able to complete and deploy a build.
+ */
 namespace ETargetPlatformReadyStatus
 {
-	/**
-	 * Flags specifiying what is needed to be able to complete and deploy a build
-	 */
-	enum Type
-	{
-		/** Ready */
-		Ready = 0,
+	/** Ready */
+	const int32 Ready = 0;
 
-		/** SDK Not Found*/
-		SDKNotFound = 1,
+	/** SDK Not Found*/
+	const int32 SDKNotFound = 1;
 
-		/** Code Build Not Supported */
-		CodeUnsupported = 2,
+	/** Code Build Not Supported */
+	const int32 CodeUnsupported = 2;
 
-		/** Plugins Not Supported */
-		PluginsUnsupported = 4,
+	/** Plugins Not Supported */
+	const int32 PluginsUnsupported = 4;
 
-		/** Signing Key Not Found */
-		SigningKeyNotFound = 8,
+	/** Signing Key Not Found */
+	const int32 SigningKeyNotFound = 8;
 
-		/** Provision Not Found */
-		ProvisionNotFound = 16,
+	/** Provision Not Found */
+	const int32 ProvisionNotFound = 16;
 
-		/** Manifest Not Found */
-		ManifestNotFound = 32,
-	};
+	/** Manifest Not Found */
+	const int32 ManifestNotFound = 32;
 };
 
 
 /**
- * ITargetPlatform, abstraction for cooking platforms and enumerating actual target devices
-**/
+ * Interface for target platforms.
+ *
+ * This interface provides an abstraction for cooking platforms and enumerating actual target devices.
+ */
 class ITargetPlatform
 {
 public:
@@ -144,6 +114,17 @@ public:
 	 */
 	virtual FText DisplayName( ) const = 0;
 
+	/**
+	 * Checks whether the platform's build requirements are met so that we can do things like
+	 * package for the platform
+	 *
+	 * @param ProjectPath Path to the project
+	 * @param bProjectHasCode true if the project has code, and therefore any compilation based SDK requirements should be checked
+	 * @param OutDocumentationPath Let's the platform tell the editor a path to show some information about how to fix any problem
+	 * @return Readiness status
+	 */
+	virtual int DoesntHaveRequirements(const FString& ProjectPath, bool bProjectHasCode, FString& OutDocumentationPath) const = 0;
+
 	/**
 	 * Returns the information about this platform
 	 */
@@ -208,6 +189,13 @@ public:
 	 */
 	virtual bool HasEditorOnlyData( ) const = 0;
 
+	/**
+	 * Checks whether this platform is only a client (and must connect to a server to run).
+	 *
+	 * @return true if this platform must connect to a server.
+	 */
+	virtual bool IsClientOnly( ) const = 0;
+
 	/**
 	 * Checks whether this platform is little endian.
 	 *
@@ -233,11 +221,14 @@ public:
 	virtual bool IsServerOnly( ) const = 0;
 
 	/**
-	 * Checks whether this platform is only a client (and must connect to a server to run).
+	 * Checks whether the platform's SDK requirements are met so that we can do things like
+	 * package for the platform
 	 *
-	 * @return true if this platform must connect to a server.
+	 * @param bProjectHasCode true if the project has code, and therefore any compilation based SDK requirements should be checked
+	 * @param OutDocumentationPath Let's the platform tell the editor a path to show some documentation about how to set up the SDK
+	 * @return true if the platform is ready for use
 	 */
-	virtual bool IsClientOnly( ) const = 0;
+	virtual bool IsSdkInstalled(bool bProjectHasCode, FString& OutDocumentationPath) const = 0;
 
 	/**
 	 * Returns the maximum bones the platform supports.
@@ -260,6 +251,11 @@ public:
 	 */
 	virtual bool RequiresUserCredentials() const = 0;
 
+	/**
+	 * Returns true if the platform supports the AutoSDK system
+	 */
+	virtual bool SupportsAutoSDK() const = 0;
+
 	/**
 	 * Checks whether this platform supports the specified build target, i.e. Game or Editor.
 	 *
@@ -268,41 +264,13 @@ public:
 	 */
 	virtual bool SupportsBuildTarget( EBuildTargets::Type BuildTarget ) const = 0;
 
-	/**
-	 * Returns true if the platform supports the AutoSDK system
-	 */
-	virtual bool SupportsAutoSDK() const = 0;
-
 	/**
 	 * Checks whether the target platform supports the specified feature.
 	 *
 	 * @param Feature The feature to check.
 	 * @return true if the feature is supported, false otherwise.
 	 */
-	virtual bool SupportsFeature( ETargetPlatformFeatures::Type Feature ) const = 0;
-
-
-	/**
-	 * Checks whether the platform's SDK requirements are met so that we can do things like
-	 * package for the platform
-	 *
-	 * @param bProjectHasCode true if the project has code, and therefore any compilation based SDK requirements should be checked
-	 * @param OutDocumentationPath Let's the platform tell the editor a path to show some documentation about how to set up the SDK
-	 * @return true if the platform is ready for use
-	 */
-	virtual bool IsSdkInstalled(bool bProjectHasCode, FString& OutDocumentationPath) const = 0;
-
-
-	/**
-	 * Checks whether the platform's build requirements are met so that we can do things like
-	 * package for the platform
-	 *
-	 * @param ProjectPath Path to the project
-	 * @param bProjectHasCode true if the project has code, and therefore any compilation based SDK requirements should be checked
-	 * @param OutDocumentationPath Let's the platform tell the editor a path to show some information about how to fix any problem
-	 * @return Readiness status
-	 */
-	virtual int DoesntHaveRequirements(const FString& ProjectPath, bool bProjectHasCode, FString& OutDocumentationPath) const = 0;
+	virtual bool SupportsFeature( ETargetPlatformFeatures Feature ) const = 0;
 
 #if WITH_ENGINE
 	/**
diff --git a/Engine/Source/Developer/Windows/WindowsTargetPlatform/Private/GenericWindowsTargetPlatform.h b/Engine/Source/Developer/Windows/WindowsTargetPlatform/Private/GenericWindowsTargetPlatform.h
--- a/Engine/Source/Developer/Windows/WindowsTargetPlatform/Private/GenericWindowsTargetPlatform.h
+++ b/Engine/Source/Developer/Windows/WindowsTargetPlatform/Private/GenericWindowsTargetPlatform.h
@@ -93,7 +93,7 @@ public:
 		return PLATFORM_WINDOWS && !UE_SERVER && !UE_GAME && WITH_EDITOR && HAS_EDITOR_DATA;
 	}
 
-	virtual bool SupportsFeature( ETargetPlatformFeatures::Type Feature ) const override
+	virtual bool SupportsFeature( ETargetPlatformFeatures Feature ) const override
 	{
 		// we currently do not have a build target for WindowsServer
 		if (Feature == ETargetPlatformFeatures::Packaging)
diff --git a/Engine/Source/Developer/Windows/WindowsTargetPlatform/Private/LocalPcTargetDevice.h b/Engine/Source/Developer/Windows/WindowsTargetPlatform/Private/LocalPcTargetDevice.h
--- a/Engine/Source/Developer/Windows/WindowsTargetPlatform/Private/LocalPcTargetDevice.h
+++ b/Engine/Source/Developer/Windows/WindowsTargetPlatform/Private/LocalPcTargetDevice.h
@@ -61,7 +61,7 @@ public:
 	virtual void Disconnect( )
 	{ }
 
-	virtual ETargetDeviceTypes::Type GetDeviceType( ) const override
+	virtual ETargetDeviceTypes GetDeviceType( ) const override
 	{
 		if (::GetSystemMetrics(SM_TABLETPC) != 0)
 		{
@@ -291,7 +291,7 @@ public:
 		return ProcessHandle.Close();
 	}
 
-	virtual bool SupportsFeature( ETargetDeviceFeatures::Type Feature ) const override
+	virtual bool SupportsFeature( ETargetDeviceFeatures Feature ) const override
 	{
 		switch (Feature)
 		{
diff --git a/Engine/Source/Developer/iOS/IOSTargetPlatform/Private/IOSTargetDevice.cpp b/Engine/Source/Developer/iOS/IOSTargetPlatform/Private/IOSTargetDevice.cpp
--- a/Engine/Source/Developer/iOS/IOSTargetPlatform/Private/IOSTargetDevice.cpp
+++ b/Engine/Source/Developer/iOS/IOSTargetPlatform/Private/IOSTargetDevice.cpp
@@ -24,7 +24,7 @@ int32 FIOSTargetDevice::GetProcessSnapshot(TArray<FTargetDeviceProcessInfo>& Out
 	return 0;
 }
 
-ETargetDeviceTypes::Type FIOSTargetDevice::GetDeviceType() const
+ETargetDeviceTypes FIOSTargetDevice::GetDeviceType() const
 {
 	return DeviceType;
 }
@@ -128,7 +128,7 @@ bool FIOSTargetDevice::Run(const FString& ExecutablePath, const FString& Params,
 #endif // !PLATFORM_MAC
 }
 
-bool FIOSTargetDevice::SupportsFeature(ETargetDeviceFeatures::Type Feature) const
+bool FIOSTargetDevice::SupportsFeature(ETargetDeviceFeatures Feature) const
 {
 	switch (Feature)
 	{
diff --git a/Engine/Source/Developer/iOS/IOSTargetPlatform/Private/IOSTargetDevice.h b/Engine/Source/Developer/iOS/IOSTargetPlatform/Private/IOSTargetDevice.h
--- a/Engine/Source/Developer/iOS/IOSTargetPlatform/Private/IOSTargetDevice.h
+++ b/Engine/Source/Developer/iOS/IOSTargetPlatform/Private/IOSTargetDevice.h
@@ -43,7 +43,7 @@ public:
 	virtual bool Deploy(const FString& SourceFolder, FString& OutAppId) override;
 	virtual void Disconnect() override;
 	virtual int32 GetProcessSnapshot(TArray<FTargetDeviceProcessInfo>& OutProcessInfos) override;
-	virtual ETargetDeviceTypes::Type GetDeviceType() const override;
+	virtual ETargetDeviceTypes GetDeviceType() const override;
 	virtual FTargetDeviceId GetId() const override;
 	virtual FString GetName() const override;
 	virtual FString GetOperatingSystemName() override;
@@ -55,7 +55,7 @@ public:
 	virtual bool PowerOn() override;
 	virtual bool Reboot(bool bReconnect = false) override;
 	virtual bool Run(const FString& ExecutablePath, const FString& Params, uint32* OutProcessId) override;
-	virtual bool SupportsFeature(ETargetDeviceFeatures::Type Feature) const;
+	virtual bool SupportsFeature(ETargetDeviceFeatures Feature) const;
 	virtual bool SupportsSdkVersion(const FString& VersionString) const override;
 	virtual bool TerminateProcess(const int32 ProcessId) override;
 	virtual void SetUserCredentials(const FString& UserName, const FString& UserPassword) override;
@@ -97,11 +97,11 @@ private:
 	/** Name of device */
 	FString DeviceName;
 	/** Type of device */
-	ETargetDeviceTypes::Type DeviceType;
+	ETargetDeviceTypes DeviceType;
 
 public:
 
-	void SetFeature(ETargetDeviceFeatures::Type InFeature, bool bFlag)
+	void SetFeature(ETargetDeviceFeatures InFeature, bool bFlag)
 	{
 		if(InFeature == ETargetDeviceFeatures::Reboot)
 		{
diff --git a/Engine/Source/Developer/iOS/IOSTargetPlatform/Private/IOSTargetPlatform.h b/Engine/Source/Developer/iOS/IOSTargetPlatform/Private/IOSTargetPlatform.h
--- a/Engine/Source/Developer/iOS/IOSTargetPlatform/Private/IOSTargetPlatform.h
+++ b/Engine/Source/Developer/iOS/IOSTargetPlatform/Private/IOSTargetPlatform.h
@@ -72,7 +72,7 @@ public:
 		#endif
 	}
 
-	virtual bool SupportsFeature( ETargetPlatformFeatures::Type Feature ) const override
+	virtual bool SupportsFeature( ETargetPlatformFeatures Feature ) const override
 	{
 		if (Feature == ETargetPlatformFeatures::Packaging)
 		{
