TargetPlatform: Documentation updates, cosmetic changes

diff --git a/Engine/Source/Developer/TargetPlatform/Public/Common/TargetPlatformBase.h b/Engine/Source/Developer/TargetPlatform/Public/Common/TargetPlatformBase.h
--- a/Engine/Source/Developer/TargetPlatform/Public/Common/TargetPlatformBase.h
+++ b/Engine/Source/Developer/TargetPlatform/Public/Common/TargetPlatformBase.h
@@ -13,19 +13,19 @@ class FTargetPlatformBase
 {
 public:
 
-	// Begin ITargetPlatform interface
+	// ITargetPlatform interface
 
 	virtual bool AddDevice( const FString& DeviceName, bool bDefault ) override
 	{
 		return false;
 	}
 
-	virtual FText DisplayName( ) const override
+	virtual FText DisplayName() const override
 	{
 		return PlatformInfo->DisplayName;
 	}
 
-	virtual const PlatformInfo::FPlatformInfo& GetPlatformInfo( ) const override
+	virtual const PlatformInfo::FPlatformInfo& GetPlatformInfo() const override
 	{
 		return *PlatformInfo;
 	}
@@ -196,9 +196,8 @@ public:
 		return 0.0f;
 	}
 
-	// End ITargetPlatform interface
-
 protected:
+
 	FTargetPlatformBase(const PlatformInfo::FPlatformInfo *const InPlatformInfo)
 		: PlatformInfo(InPlatformInfo)
 	{
@@ -222,7 +221,7 @@ class TTargetPlatformBase
 public:
 
 	/** Default constructor. */
-	TTargetPlatformBase( )
+	TTargetPlatformBase()
 		: FTargetPlatformBase( PlatformInfo::FindPlatformInfo(TPlatformProperties::PlatformName()) )
 	{
 		// HasEditorOnlyData and RequiresCookedData are mutually exclusive.
@@ -233,37 +232,37 @@ public:
 
 	// ITargetPlatform interface
 
-	virtual bool HasEditorOnlyData( ) const override
+	virtual bool HasEditorOnlyData() const override
 	{
 		return TPlatformProperties::HasEditorOnlyData();
 	}
 
-	virtual bool IsLittleEndian( ) const override
+	virtual bool IsLittleEndian() const override
 	{
 		return TPlatformProperties::IsLittleEndian();
 	}
 
-	virtual bool IsServerOnly( ) const override
+	virtual bool IsServerOnly() const override
 	{
 		return TPlatformProperties::IsServerOnly();
 	}
 
-	virtual bool IsClientOnly( ) const override
+	virtual bool IsClientOnly() const override
 	{
 		return TPlatformProperties::IsClientOnly();
 	}
 
-	virtual FString PlatformName( ) const override
+	virtual FString PlatformName() const override
 	{
 		return FString(TPlatformProperties::PlatformName());
 	}
 
-	virtual FString IniPlatformName( ) const override
+	virtual FString IniPlatformName() const override
 	{
 		return FString(TPlatformProperties::IniPlatformName());
 	}
 
-	virtual bool RequiresCookedData( ) const override
+	virtual bool RequiresCookedData() const override
 	{
 		return TPlatformProperties::RequiresCookedData();
 	}
@@ -321,7 +320,7 @@ public:
 		return false;
 	}
 	
-	virtual uint32 MaxGpuSkinBones( ) const override
+	virtual uint32 MaxGpuSkinBones() const override
 	{
 		return TPlatformProperties::MaxGpuSkinBones();
 	}
diff --git a/Engine/Source/Developer/TargetPlatform/Public/Interfaces/ITargetPlatform.h b/Engine/Source/Developer/TargetPlatform/Public/Interfaces/ITargetPlatform.h
--- a/Engine/Source/Developer/TargetPlatform/Public/Interfaces/ITargetPlatform.h
+++ b/Engine/Source/Developer/TargetPlatform/Public/Interfaces/ITargetPlatform.h
@@ -112,7 +112,7 @@ public:
 	 *
 	 * @see PlatformName
 	 */
-	virtual FText DisplayName( ) const = 0;
+	virtual FText DisplayName() const = 0;
 
 	/**
 	 * Checks whether the platform's build requirements are met so that we can do things like package for the platform.
@@ -120,21 +120,21 @@ public:
 	 * @param ProjectPath Path to the project.
 	 * @param bProjectHasCode true if the project has code, and therefore any compilation based SDK requirements should be checked.
 	 * @param OutDocumentationPath Let's the platform tell the editor a path to show some information about how to fix any problem.
-	 * @return Readiness status.
+	 * @return A mask of ETargetPlatformReadyStatus flags to indicate missing requirements, or 0 if all requirements are met.
 	 */
 	virtual int32 CheckRequirements(const FString& ProjectPath, bool bProjectHasCode, FString& OutDocumentationPath) const = 0;
 
 	/**
 	 * Returns the information about this platform
 	 */
-	virtual const PlatformInfo::FPlatformInfo& GetPlatformInfo( ) const = 0;
+	virtual const PlatformInfo::FPlatformInfo& GetPlatformInfo() const = 0;
 
 	/**
 	 * Gets the platform's INI name (so an offline tool can load the INI for the given target platform).
 	 *
 	 * @see PlatformName
 	 */
-	virtual FString IniPlatformName( ) const = 0;
+	virtual FString IniPlatformName() const = 0;
 
 	/**
 	 * Enables/Disable the device check
@@ -153,7 +153,7 @@ public:
 	 *
 	 * @return Compression method.
 	 */
-	virtual ECompressionFlags GetBaseCompressionMethod( ) const = 0;
+	virtual ECompressionFlags GetBaseCompressionMethod() const = 0;
 
 	/** 
 	 * Generates a platform specific asset manifest given an array of FAssetData.
@@ -171,7 +171,7 @@ public:
 	 *
 	 * @return Default device.
 	 */
-	virtual ITargetDevicePtr GetDefaultDevice( ) const = 0;
+	virtual ITargetDevicePtr GetDefaultDevice() const = 0;
 
 	/** 
 	 * Gets an interface to the specified device.
@@ -186,21 +186,21 @@ public:
 	 *
 	 * @return true if this platform has editor only data, false otherwise.
 	 */
-	virtual bool HasEditorOnlyData( ) const = 0;
+	virtual bool HasEditorOnlyData() const = 0;
 
 	/**
 	 * Checks whether this platform is only a client (and must connect to a server to run).
 	 *
 	 * @return true if this platform must connect to a server.
 	 */
-	virtual bool IsClientOnly( ) const = 0;
+	virtual bool IsClientOnly() const = 0;
 
 	/**
 	 * Checks whether this platform is little endian.
 	 *
 	 * @return true if this platform is little-endian, false otherwise.
 	 */
-	virtual bool IsLittleEndian( ) const = 0;
+	virtual bool IsLittleEndian() const = 0;
 
 	/**
 	 * Checks whether this platform is the platform that's currently running.
@@ -210,14 +210,14 @@ public:
 	 *
 	 * @return true if this platform is running, false otherwise.
 	 */
-	virtual bool IsRunningPlatform( ) const = 0;
+	virtual bool IsRunningPlatform() const = 0;
 
 	/**
 	 * Checks whether this platform is only a server.
 	 *
 	 * @return true if this platform has no graphics or audio, etc, false otherwise.
 	 */
-	virtual bool IsServerOnly( ) const = 0;
+	virtual bool IsServerOnly() const = 0;
 
 	/**
 	 * Checks whether the platform's SDK requirements are met so that we can do things like
@@ -234,14 +234,14 @@ public:
 	 *
 	 * @return the maximum bones the platform supports.
 	 */
-	virtual uint32 MaxGpuSkinBones( ) const = 0;
+	virtual uint32 MaxGpuSkinBones() const = 0;
 
 	/**
 	 * Checks whether this platform requires cooked data (typically console platforms).
 	 *
 	 * @return true if this platform requires cooked data, false otherwise.
 	 */
-	virtual bool RequiresCookedData( ) const = 0;
+	virtual bool RequiresCookedData() const = 0;
 
 	/**
 	 * Checks whether this platform requires user credentials (typically server platforms).
@@ -321,7 +321,7 @@ public:
 	 *
 	 * @return A texture LOD settings structure.
 	 */
-	virtual const struct FTextureLODSettings& GetTextureLODSettings( ) const = 0;
+	virtual const struct FTextureLODSettings& GetTextureLODSettings() const = 0;
 
 	/**
 	 * Gets the static mesh LOD settings used by this platform.
