comment fix

diff --git a/Engine/Source/Developer/TargetPlatform/Public/Interfaces/ITextureFormat.h b/Engine/Source/Developer/TargetPlatform/Public/Interfaces/ITextureFormat.h
--- a/Engine/Source/Developer/TargetPlatform/Public/Interfaces/ITextureFormat.h
+++ b/Engine/Source/Developer/TargetPlatform/Public/Interfaces/ITextureFormat.h
@@ -24,7 +24,7 @@ class ITextureFormat
 public:
 
 	/**
-	 * Checks whether this audio format can compress in parallel.
+	 * Checks whether this texture format can compress in parallel.
 	 *
 	 * @return true if parallel compression is supported, false otherwise.
 	 */
