- Fixing variable case in SDKSetup function

diff --git a/Engine/Source/Developer/TargetPlatform/Private/TargetPlatformManagerModule.cpp b/Engine/Source/Developer/TargetPlatform/Private/TargetPlatformManagerModule.cpp
--- a/Engine/Source/Developer/TargetPlatform/Private/TargetPlatformManagerModule.cpp
+++ b/Engine/Source/Developer/TargetPlatform/Private/TargetPlatformManagerModule.cpp
@@ -826,7 +826,7 @@ protected:
 		return SetupSDKStatus(TEXT(""));
 	}
 
-	bool SetupSDKStatus(FString targetPlatforms)
+	bool SetupSDKStatus(FString TargetPlatforms)
 	{
 		// run UBT with -validate -allplatforms and read the output
 #if PLATFORM_MAC
@@ -839,7 +839,7 @@ protected:
 #endif
 
 		// Allow for only a subset of platforms to be reparsed - needed when kicking a change from the UI
-		CommandLine += targetPlatforms.IsEmpty() ? TEXT(" -allplatforms") : (TEXT(" -platforms=") + targetPlatforms);
+		CommandLine += TargetPlatforms.IsEmpty() ? TEXT(" -allplatforms") : (TEXT(" -platforms=") + TargetPlatforms);
 
 		TSharedPtr<FMonitoredProcess> UBTProcess = MakeShareable(new FMonitoredProcess(CmdExe, CommandLine, true));
 		UBTProcess->OnOutput().BindStatic(&FTargetPlatformManagerModule::OnStatusOutput);
