Fix incorrect case in header file name.

diff --git a/Engine/Source/Developer/TargetPlatform/Public/Interfaces/ITargetPlatformManagerModule.h b/Engine/Source/Developer/TargetPlatform/Public/Interfaces/ITargetPlatformManagerModule.h
--- a/Engine/Source/Developer/TargetPlatform/Public/Interfaces/ITargetPlatformManagerModule.h
+++ b/Engine/Source/Developer/TargetPlatform/Public/Interfaces/ITargetPlatformManagerModule.h
@@ -2,7 +2,7 @@
 
 #pragma once
 
-#include "TargetDeviceID.h"
+#include "TargetDeviceId.h"
 #include "ITargetDevice.h"
 
 /**
