Enable Rocket builds for HTML5 on Mac

diff --git a/Engine/Source/Developer/TargetPlatform/TargetPlatform.Build.cs b/Engine/Source/Developer/TargetPlatform/TargetPlatform.Build.cs
--- a/Engine/Source/Developer/TargetPlatform/TargetPlatform.Build.cs
+++ b/Engine/Source/Developer/TargetPlatform/TargetPlatform.Build.cs
@@ -89,11 +89,7 @@ public class TargetPlatform : ModuleRules
 					DynamicallyLoadedModuleNames.Add("Android_ETC1TargetPlatform");
 					DynamicallyLoadedModuleNames.Add("Android_ETC2TargetPlatform");
 					DynamicallyLoadedModuleNames.Add("IOSTargetPlatform");
-
-					if (!UnrealBuildTool.UnrealBuildTool.RunningRocket() && !UnrealBuildTool.UnrealBuildTool.BuildingRocket())
-					{
-						DynamicallyLoadedModuleNames.Add("HTML5TargetPlatform");
-					}
+					DynamicallyLoadedModuleNames.Add("HTML5TargetPlatform");
 				}
 			}
             else if (Target.Platform == UnrealTargetPlatform.Linux)
