Fix SDK setup check for Linux.

diff --git a/Engine/Source/Developer/TargetPlatform/Private/TargetPlatformManagerModule.cpp b/Engine/Source/Developer/TargetPlatform/Private/TargetPlatformManagerModule.cpp
--- a/Engine/Source/Developer/TargetPlatform/Private/TargetPlatformManagerModule.cpp
+++ b/Engine/Source/Developer/TargetPlatform/Private/TargetPlatformManagerModule.cpp
@@ -829,14 +829,28 @@ protected:
 	bool SetupSDKStatus(FString TargetPlatforms)
 	{
 		// run UBT with -validate -allplatforms and read the output
-#if PLATFORM_MAC
-		FString CmdExe = TEXT("/bin/sh");
-        FString ScriptPath = FPaths::ConvertRelativePathToFull(FPaths::EngineDir() / TEXT("Build/BatchFiles/Mac/RunMono.sh"));
-		FString CommandLine = TEXT("\"") + ScriptPath + TEXT("\" \"") + FPaths::ConvertRelativePathToFull(FPaths::EngineDir() / TEXT("Binaries/DotNet/UnrealBuildTool.exe")) + TEXT("\" -validateplatform");
-#else
-		FString CmdExe = FPaths::ConvertRelativePathToFull(FPaths::EngineDir() / TEXT("Binaries/DotNet/UnrealBuildTool.exe"));
-		FString CommandLine = TEXT("-validateplatform");
-#endif
+		FString CmdExe, CommandLine;
+		
+		if (PLATFORM_MAC)
+		{
+			CmdExe = TEXT("/bin/sh");
+			FString ScriptPath = FPaths::ConvertRelativePathToFull(FPaths::EngineDir() / TEXT("Build/BatchFiles/Mac/RunMono.sh"));
+			CommandLine = TEXT("\"") + ScriptPath + TEXT("\" \"") + FPaths::ConvertRelativePathToFull(FPaths::EngineDir() / TEXT("Binaries/DotNET/UnrealBuildTool.exe")) + TEXT("\" -validateplatform");
+		}
+		else if (PLATFORM_WINDOWS)
+		{
+			CmdExe = FPaths::ConvertRelativePathToFull(FPaths::EngineDir() / TEXT("Binaries/DotNet/UnrealBuildTool.exe"));
+			CommandLine = TEXT("-validateplatform");
+		}
+		else if (PLATFORM_LINUX)
+		{
+			CmdExe = TEXT("/usr/bin/mono");
+			CommandLine = TEXT("\"") + FPaths::ConvertRelativePathToFull(FPaths::EngineDir() / TEXT("Binaries/DotNET/UnrealBuildTool.exe")) + TEXT("\" -validateplatform");
+		}
+		else
+		{
+			checkf(false, TEXT("FTargetPlatformManagerModule::SetupSDKStatus(): Unsupported platform!"));
+		}
 
 		// Allow for only a subset of platforms to be reparsed - needed when kicking a change from the UI
 		CommandLine += TargetPlatforms.IsEmpty() ? TEXT(" -allplatforms") : (TEXT(" -platforms=") + TargetPlatforms);
