Stats - Added a few stats to track ini loading

diff --git a/Engine/Source/Developer/TargetPlatform/Private/TargetPlatformManagerModule.cpp b/Engine/Source/Developer/TargetPlatform/Private/TargetPlatformManagerModule.cpp
--- a/Engine/Source/Developer/TargetPlatform/Private/TargetPlatformManagerModule.cpp
+++ b/Engine/Source/Developer/TargetPlatform/Private/TargetPlatformManagerModule.cpp
@@ -35,6 +35,8 @@ public:
 		// AutoSDKs only enabled if UE_SDKS_ROOT is set.
 		if (IsAutoSDKsEnabled())
 		{					
+			DECLARE_SCOPE_CYCLE_COUNTER( TEXT( "FTargetPlatformManagerModule.StartAutoSDK" ), STAT_FTargetPlatformManagerModule_StartAutoSDK, STATGROUP_TargetPlatform );
+
 			// amortize UBT cost by calling it once for all platforms, rather than once per platform.
 			FString UBTParams(TEXT("-autosdkonly"));
 			int32 UBTReturnCode = -1;
@@ -555,6 +557,8 @@ protected:
 	/** Discovers the available target platforms. */
 	void DiscoverAvailablePlatforms()
 	{
+		DECLARE_SCOPE_CYCLE_COUNTER( TEXT( "FTargetPlatformManagerModule::DiscoverAvailablePlatforms" ), STAT_FTargetPlatformManagerModule_DiscoverAvailablePlatforms, STATGROUP_TargetPlatform );
+
 		Platforms.Empty(Platforms.Num());
 
 		TArray<FName> Modules;
@@ -828,6 +832,8 @@ protected:
 
 	bool SetupSDKStatus(FString TargetPlatforms)
 	{
+		DECLARE_SCOPE_CYCLE_COUNTER( TEXT( "FTargetPlatformManagerModule::SetupSDKStatus" ), STAT_FTargetPlatformManagerModule_SetupSDKStatus, STATGROUP_TargetPlatform );
+
 		// run UBT with -validate -allplatforms and read the output
 		FString CmdExe, CommandLine;
 		
diff --git a/Engine/Source/Runtime/Core/Private/Misc/ConfigCacheIni.cpp b/Engine/Source/Runtime/Core/Private/Misc/ConfigCacheIni.cpp
--- a/Engine/Source/Runtime/Core/Private/Misc/ConfigCacheIni.cpp
+++ b/Engine/Source/Runtime/Core/Private/Misc/ConfigCacheIni.cpp
@@ -2949,6 +2949,7 @@ bool FConfigCacheIni::LoadGlobalIniFile(FString& FinalIniFilename, const TCHAR*
 
 void FConfigCacheIni::LoadLocalIniFile(FConfigFile& ConfigFile, const TCHAR* IniName, bool bGenerateDestIni, const TCHAR* Platform, const TCHAR* GameName, const bool bForceReload )
 {
+	DECLARE_SCOPE_CYCLE_COUNTER( TEXT( "FConfigCacheIni::LoadLocalIniFile" ), STAT_FConfigCacheIni_LoadLocalIniFile, STATGROUP_LoadTime );
 
 	LoadExternalIniFile( ConfigFile, IniName, *FPaths::EngineConfigDir(), *FPaths::SourceConfigDir(), bGenerateDestIni, Platform, GameName, bForceReload );
 
diff --git a/Engine/Source/Runtime/Core/Private/Misc/RemoteConfigIni.cpp b/Engine/Source/Runtime/Core/Private/Misc/RemoteConfigIni.cpp
--- a/Engine/Source/Runtime/Core/Private/Misc/RemoteConfigIni.cpp
+++ b/Engine/Source/Runtime/Core/Private/Misc/RemoteConfigIni.cpp
@@ -445,6 +445,8 @@ bool IsUsingLocalIniFile(const TCHAR* FilenameToLoad, const TCHAR* IniFileName)
 /** Contains the logic for processing config files, local or remote */
 void ProcessIniContents(const TCHAR* FilenameToLoad, const TCHAR* IniFileName, FConfigFile* Config, bool bDoEmptyConfig, bool bDoCombine, bool bDoWrite)
 {
+	DECLARE_SCOPE_CYCLE_COUNTER( TEXT( "ProcessIniContents" ), STAT_ProcessIniContents, STATGROUP_LoadTime );
+
 	check(FilenameToLoad);
 	check(IniFileName);
 	check(Config);
