UE4 - merged fix to eliminate redundant target platform loads

diff --git a/Engine/Source/Developer/TargetPlatform/Private/TargetPlatformManagerModule.cpp b/Engine/Source/Developer/TargetPlatform/Private/TargetPlatformManagerModule.cpp
--- a/Engine/Source/Developer/TargetPlatform/Private/TargetPlatformManagerModule.cpp
+++ b/Engine/Source/Developer/TargetPlatform/Private/TargetPlatformManagerModule.cpp
@@ -28,7 +28,9 @@ public:
 
 	/** Default constructor. */
 	FTargetPlatformManagerModule()
-		: bRestrictFormatsToRuntimeOnly(false), bForceCacheUpdate(true)
+		: bRestrictFormatsToRuntimeOnly(false)
+		, bForceCacheUpdate(true)
+		, bIgnoreFirstDelegateCall(true)
 	{
 #if AUTOSDKS_ENABLED		
 		
@@ -62,7 +64,7 @@ public:
 #endif
 
 		SetupSDKStatus();
-		GetTargetPlatforms();
+		//GetTargetPlatforms(); redudant with next call
 		GetActiveTargetPlatforms();
 		GetAudioFormats();
 		GetTextureFormats();
@@ -88,7 +90,7 @@ public:
 		bForceCacheUpdate = true;
 
 		SetupSDKStatus();
-		GetTargetPlatforms();
+		//GetTargetPlatforms(); redudant with next call
 		GetActiveTargetPlatforms();
 		GetAudioFormats();
 		GetTextureFormats();
@@ -587,6 +589,7 @@ protected:
 					// this setup will become faster after TTP 341897 is complete.					
 					if (SetupAndValidateAutoSDK(Platform->GetPlatformInfo().AutoSDKPath))
 					{
+						UE_LOG(LogTemp, Display, TEXT("Loaded TP %s"), *Modules[Index].ToString());
 						Platforms.Add(Platform);
 					}
 				}
@@ -939,10 +942,11 @@ private:
 
 	void ModulesChangesCallback(FName ModuleName, EModuleChangeReason ReasonForChange)
 	{
-		if (ModuleName.ToString().Contains(TEXT("TargetPlatform")) )
+		if (!bIgnoreFirstDelegateCall && ModuleName.ToString().Contains(TEXT("TargetPlatform")) && !ModuleName.ToString().Contains(TEXT("ProjectTargetPlatformEditor")))
 		{
 			Invalidate();
 		}
+		bIgnoreFirstDelegateCall = false;
 	}
 	
 	static FString SDKStatusMessage;
@@ -960,6 +964,9 @@ private:
 	// in case of a module reload of a TargetPlatform-Module.
 	bool bForceCacheUpdate;
 
+	// Flag to avoid redunant reloads
+	bool bIgnoreFirstDelegateCall;
+
 	// Holds the list of discovered platforms.
 	TArray<ITargetPlatform*> Platforms;
 
