Added additional error information for commandline users when entering incorrect targetplatform.

diff --git a/Engine/Source/Developer/TargetPlatform/Private/TargetPlatformManagerModule.cpp b/Engine/Source/Developer/TargetPlatform/Private/TargetPlatformManagerModule.cpp
--- a/Engine/Source/Developer/TargetPlatform/Private/TargetPlatformManagerModule.cpp
+++ b/Engine/Source/Developer/TargetPlatform/Private/TargetPlatformManagerModule.cpp
@@ -213,6 +213,7 @@ public:
 					{
 						// An invalid platform was specified...
 						// Inform the user and exit.
+						UE_LOG(LogTargetPlatformManager, Error, TEXT("Invalid target platform specified (%s). Available = { %s } "), *PlatformStr, *AvailablePlatforms);
 						UE_LOG(LogTargetPlatformManager, Fatal, TEXT("Invalid target platform specified (%s). Available = { %s } "), *PlatformStr, *AvailablePlatforms);
 					}
 				}
