Disabled compressed cooked packages for all platforms on launch on as it causes launch on to stream slowly.

diff --git a/Engine/Source/Developer/TargetPlatform/Public/Common/TargetPlatformBase.h b/Engine/Source/Developer/TargetPlatform/Public/Common/TargetPlatformBase.h
--- a/Engine/Source/Developer/TargetPlatform/Public/Common/TargetPlatformBase.h
+++ b/Engine/Source/Developer/TargetPlatform/Public/Common/TargetPlatformBase.h
@@ -357,7 +357,7 @@ public:
 		case ETargetPlatformFeatures::UserCredentials:
 			break;
 		case ETargetPlatformFeatures::ShouldUseCompressedCookedPackages:
-			return true;
+			return false;
 		}
 
 		return false;
