Fixed int32 returned as a boolean

diff --git a/Engine/Source/Developer/TargetPlatform/Private/TargetPlatformBase.cpp b/Engine/Source/Developer/TargetPlatform/Private/TargetPlatformBase.cpp
--- a/Engine/Source/Developer/TargetPlatform/Private/TargetPlatformBase.cpp
+++ b/Engine/Source/Developer/TargetPlatform/Private/TargetPlatformBase.cpp
@@ -19,13 +19,13 @@ bool FTargetPlatformBase::UsesDBuffer() const
 bool FTargetPlatformBase::UsesBasePassVelocity() const
 {
 	static IConsoleVariable* CVar = IConsoleManager::Get().FindConsoleVariable(TEXT("r.BasePassOutputsVelocity"));
-	return CVar ? (CVar->GetInt()) : false;
+	return CVar ? (CVar->GetInt() != 0) : false;
 }
 
 bool FTargetPlatformBase::UsesSelectiveBasePassOutputs() const
 {
 	static IConsoleVariable* CVar = IConsoleManager::Get().FindConsoleVariable(TEXT("r.SelectiveBasePassOutputs"));
-	return CVar ? (CVar->GetInt()) : false;
+	return CVar ? (CVar->GetInt() != 0) : false;
 }
 
 TSharedPtr<IDeviceManagerCustomPlatformWidgetCreator> FTargetPlatformBase::GetCustomWidgetCreator() const
