Fix inability to load dependent modules located in platform subdirectories.

diff --git a/Engine/Source/Developer/TargetPlatform/Private/TargetPlatformManagerModule.cpp b/Engine/Source/Developer/TargetPlatform/Private/TargetPlatformManagerModule.cpp
--- a/Engine/Source/Developer/TargetPlatform/Private/TargetPlatformManagerModule.cpp
+++ b/Engine/Source/Developer/TargetPlatform/Private/TargetPlatformManagerModule.cpp
@@ -636,6 +636,16 @@ protected:
 #endif
 		FModuleManager::Get().FindModules(*ModuleWildCard, Modules);
 
+#if !IS_MONOLITHIC
+		// Find all module subdirectories and add them so we can load dependent modules for target platform modules
+		TArray<FString> ModuleSubdirs;
+		IFileManager::Get().FindFilesRecursive(ModuleSubdirs, *FPlatformProcess::GetModulesDirectory(), TEXT("*"), false, true);
+		for (const FString& ModuleSubdir : ModuleSubdirs)
+		{
+			FModuleManager::Get().AddBinariesDirectory(*ModuleSubdir, false);
+		}
+#endif
+
 		// remove this module from the list
 		Modules.Remove(FName(TEXT("TargetPlatform")));
 
