Introduced some additional target platform utilities to help with asset cook optimizations

* We now assign each ITargetPlatform a zero-based ordinal value
* Introduced FTargetPlatform and FTargetPlatformSet types to help store platform references and platform sets efficiently.

These are not currently used in the engine but are designed to replace the existing ITargetPlatform/string/FName representations in the cooking data structures.

#rn minor

diff --git a/Engine/Source/Developer/TargetPlatform/Private/TargetPlatformManagerModule.cpp b/Engine/Source/Developer/TargetPlatform/Private/TargetPlatformManagerModule.cpp
--- a/Engine/Source/Developer/TargetPlatform/Private/TargetPlatformManagerModule.cpp
+++ b/Engine/Source/Developer/TargetPlatform/Private/TargetPlatformManagerModule.cpp
@@ -37,6 +37,36 @@ DEFINE_LOG_CATEGORY_STATIC(LogTargetPlatformManager, Log, All);
 	#define AUTOSDKS_ENABLED 0
 #endif
 
+static const size_t MaxPlatformCount = 64;		// In the unlikely event that someone bumps this please note that there's
+												// an implicit assumption that there won't be more than 64 unique target
+												// platforms in the FTargetPlatformSet code since it uses one bit of an
+												// uint64 per platform.
+
+static const ITargetPlatform* TargetPlatformArray[MaxPlatformCount];
+
+static int32 PlatformCounter = 0;
+
+int32 
+ITargetPlatform::AssignPlatformOrdinal(const ITargetPlatform& Platform)
+{
+	check(PlatformCounter < MaxPlatformCount);
+
+	const int32 Ordinal = PlatformCounter++;
+
+	check(TargetPlatformArray[Ordinal] == nullptr);
+	TargetPlatformArray[Ordinal] = &Platform;
+
+	return Ordinal;
+}
+
+const ITargetPlatform* 
+ITargetPlatform::GetPlatformFromOrdinal(int32 Ordinal)
+{
+	check(Ordinal < PlatformCounter);
+
+	return TargetPlatformArray[Ordinal];
+}
+
 
 /**
  * Module for the target platform manager
diff --git a/Engine/Source/Developer/TargetPlatform/Public/Common/TargetPlatformBase.h b/Engine/Source/Developer/TargetPlatform/Public/Common/TargetPlatformBase.h
--- a/Engine/Source/Developer/TargetPlatform/Public/Common/TargetPlatformBase.h
+++ b/Engine/Source/Developer/TargetPlatform/Public/Common/TargetPlatformBase.h
@@ -98,16 +98,24 @@ public:
 		return DEFAULT_ZLIB_BIT_WINDOW;
 	}
 
+	virtual int32 GetPlatformOrdinal() const override
+	{
+		return PlatformOrdinal;
+	}
+
 protected:
 
 	FTargetPlatformBase(const PlatformInfo::FPlatformInfo *const InPlatformInfo)
 		: PlatformInfo(InPlatformInfo)
 	{
 		check(PlatformInfo);
+
+		PlatformOrdinal = AssignPlatformOrdinal(*this);
 	}
 
 	/** Information about this platform */
 	const PlatformInfo::FPlatformInfo *PlatformInfo;
+	int32 PlatformOrdinal;
 };
 
 
diff --git a/Engine/Source/Developer/TargetPlatform/Public/Interfaces/ITargetPlatform.h b/Engine/Source/Developer/TargetPlatform/Public/Interfaces/ITargetPlatform.h
--- a/Engine/Source/Developer/TargetPlatform/Public/Interfaces/ITargetPlatform.h
+++ b/Engine/Source/Developer/TargetPlatform/Public/Interfaces/ITargetPlatform.h
@@ -250,7 +250,7 @@ public:
 	virtual bool IsServerOnly() const = 0;
 
 	/**
-	* Checks whether this platform suports shader compilation over XGE interface.
+	* Checks whether this platform supports shader compilation over XGE interface.
 	*
 	* @return true if this platform can distribute shader compilation threads with XGE.
 	*/
@@ -429,10 +429,26 @@ public:
 	virtual bool SendLowerCaseFilePaths() const = 0;
 
 	/**
-	* Project settings to check to determine if a build should occurr
-	*/
+	 * Project settings to check to determine if a build should occur
+	 */
 	virtual void GetBuildProjectSettingKeys(FString& OutSection, TArray<FString>& InBoolKeys, TArray<FString>& InIntKeys, TArray<FString>& InStringKeys) const = 0;
 
+	/**
+	 * Get unique integer identifier for this platform.
+	 *
+	 * The implementation will assign an ordinal to each target platform at startup, assigning
+	 * a value of 0, 1, 2, etc in order to make the ordinals usable as array / bit mask indices.
+	 *
+	 * @return int32 A unique integer which may be used to identify target platform during the 
+	 *               current session only (note: not stable across runs).
+	 */
+	virtual int32 GetPlatformOrdinal() const = 0;
+
+	/**
+	 * Given a platform ordinal number, returns the corresponding ITargetPlatform instance
+	 */
+	TARGETPLATFORM_API static const ITargetPlatform* GetPlatformFromOrdinal(int32 Ordinal);
+
 public:
 
 	/**
@@ -451,4 +467,143 @@ public:
 
 	/** Virtual destructor. */
 	virtual ~ITargetPlatform() { }
+
+protected:
+	static TARGETPLATFORM_API int32 AssignPlatformOrdinal(const ITargetPlatform& Platform);
+};
+
+/**
+ * Target platform identifier
+ *
+ * This is really just a wrapper around an integer ordinal value, to prevent
+ * accidental mix-ups with other classes of integers. It also provides more
+ * context to a reader of the code.
+ *
+ * @see ITargetPlatform::GetPlatformOrdinal()
+ */
+
+struct FTargetPlatform
+{
+	inline FTargetPlatform(const ITargetPlatform& Platform)
+	{
+		Ordinal = Platform.GetPlatformOrdinal();
+	}
+
+	inline uint32 GetOrdinal() const { return Ordinal;  }
+
+	bool operator<(const FTargetPlatform& Other) const
+	{
+		return Ordinal < Other.Ordinal;
+	}
+
+	bool operator==(const FTargetPlatform& Other) const
+	{
+		return Ordinal == Other.Ordinal;
+	}
+
+private:
+	uint32 Ordinal = 0;
+};
+
+inline uint32 GetTypeHash(const FTargetPlatform& Key) { return Key.GetOrdinal(); }
+
+/**
+ * Target platform set implementation using bitmask for compactness
+ */
+class FTargetPlatformSet
+{
+public:
+	inline void Add(const FTargetPlatform& Platform)
+	{
+		const uint32 Ordinal = Platform.GetOrdinal();
+
+		Mask |= 1ull << Ordinal;
+	}
+
+	inline void Remove(const FTargetPlatform& Platform)
+	{
+		const uint32 Ordinal = Platform.GetOrdinal();
+
+		Mask &= ~(1ull << Ordinal);
+	}
+
+	/// Remove all members from the Platforms set from this set
+	inline void Remove(const FTargetPlatformSet& Platforms)
+	{
+		Mask &= ~Platforms.Mask;
+	}
+
+	/// Check if this set contains any of the members of the Other set
+	inline bool Contains(const FTargetPlatform& Platform) const
+	{
+		const uint32 Ordinal = Platform.GetOrdinal();
+
+		return !!(Mask & (1ull << Ordinal));
+	}
+
+	/// Check if this set contains any of the members of the Other set
+	bool ContainsAny(const FTargetPlatformSet& Other) const
+	{
+		return !!(Mask & Other.Mask);
+	}
+
+	bool IsEmpty() const
+	{
+		return Mask == 0;
+	}
+
+	void Merge(const FTargetPlatformSet& Other)
+	{
+		Mask |= Other.Mask;
+	}
+
+	void Clear()
+	{
+		Mask = 0;
+	}
+
+	bool operator==(const FTargetPlatformSet& Other) const
+	{
+		return this->Mask == Other.Mask;
+	}
+
+	uint32 GetHash() const
+	{
+		// This may not be awesome but I don't know that it's an actual problem in practice
+		return uint32(this->Mask ^ (this->Mask >> 32));
+	}
+
+	/**
+	 * Iterate over all set members
+	 *
+	 * @param Callback - callback accepting a const ITargetPlatform* argument
+	 */
+	template<typename Func>
+	void ForEach(Func&& Callback) const
+	{
+		// This could maybe be smarter and leverage intrinsics to directly 
+		// scan for the next set bit but I just want to make it work for 
+		// now so let's keep it simple until it shows up in a profile
+
+		uint64 IterMask = Mask;
+		int Ordinal = 0;
+
+		while(IterMask)
+		{
+			if (IterMask & 1)
+			{
+				Callback(ITargetPlatform::GetPlatformFromOrdinal(Ordinal));
+			}
+
+			IterMask >>= 1;
+			++Ordinal;
+		}
+	}
+
+	// TODO: support for ranged for? It's easy to make mistakes when refactoring
+	// regular for loops into ForEach type constructs since the semantics of any
+	// return statements does change when you put them inside a lambda.
+
+private:
+	uint64	Mask = 0;
 };
