UPR #5084: Print actual TargetPlatform name instead of module (Contributed by [REDACTED_PERSON])

diff --git a/Engine/Source/Developer/TargetPlatform/Private/TargetPlatformManagerModule.cpp b/Engine/Source/Developer/TargetPlatform/Private/TargetPlatformManagerModule.cpp
--- a/Engine/Source/Developer/TargetPlatform/Private/TargetPlatformManagerModule.cpp
+++ b/Engine/Source/Developer/TargetPlatform/Private/TargetPlatformManagerModule.cpp
@@ -659,7 +659,7 @@ protected:
 RETRY_SETUPANDVALIDATE:
 					if (SetupAndValidateAutoSDK(Platform->GetPlatformInfo().AutoSDKPath))
 					{
-						UE_LOG(LogTemp, Display, TEXT("Loaded TP %s"), *Modules[Index].ToString());
+						UE_LOG(LogTemp, Display, TEXT("Module '%s' loaded TargetPlatform '%s'"), *Modules[Index].ToString(), *Platform->PlatformName());
 						Platforms.Add(Platform);
 					}
 					else
@@ -675,7 +675,7 @@ RETRY_SETUPANDVALIDATE:
 								goto RETRY_SETUPANDVALIDATE;
 							}
 						}
-						UE_LOG(LogTemp, Display, TEXT("Failed to SetupAndValidateAutoSDK for platform %s"), *Modules[Index].ToString());
+						UE_LOG(LogTemp, Display, TEXT("Module '%s' failed to SetupAndValidateAutoSDK for platform '%s'"), *Modules[Index].ToString(), *Platform->PlatformName());
 					}
 				}
 			}
