Fix mode argument to UBT when setting up AutoSDKs.

#rnx

diff --git a/Engine/Source/Developer/TargetPlatform/Private/TargetPlatformManagerModule.cpp b/Engine/Source/Developer/TargetPlatform/Private/TargetPlatformManagerModule.cpp
--- a/Engine/Source/Developer/TargetPlatform/Private/TargetPlatformManagerModule.cpp
+++ b/Engine/Source/Developer/TargetPlatform/Private/TargetPlatformManagerModule.cpp
@@ -92,7 +92,7 @@ public:
 			// amortize UBT cost by calling it once for all platforms, rather than once per platform.
 			if (FParse::Param(FCommandLine::Get(), TEXT("Multiprocess"))==false)
 			{
-				FString UBTParams(TEXT("-SetupPlatforms"));
+				FString UBTParams(TEXT("-Mode=SetupPlatforms"));
 				int32 UBTReturnCode = -1;
 				FString UBTOutput;
 				if (!FDesktopPlatformModule::Get()->InvokeUnrealBuildToolSync(UBTParams, *GLog, true, UBTReturnCode, UBTOutput))
