Pass down the shader format FName to IShaderFormat::NotifyShaderCooked

- Some platforms have multiple sub formats per IShaderFormat

diff --git a/Engine/Source/Developer/TargetPlatform/Public/Interfaces/IShaderFormat.h b/Engine/Source/Developer/TargetPlatform/Public/Interfaces/IShaderFormat.h
--- a/Engine/Source/Developer/TargetPlatform/Public/Interfaces/IShaderFormat.h
+++ b/Engine/Source/Developer/TargetPlatform/Public/Interfaces/IShaderFormat.h
@@ -98,7 +98,7 @@ public:
 	 * Called when a shader resource is cooked, so the shader format can perform platform-specific operations on the debug data.
 	 * Does nothing on platforms that make no use of the platform debug data.
 	 */
-	virtual void NotifyShaderCooked(const TArray<uint8>& PlatformDebugData) const { }
+	virtual void NotifyShaderCooked(const TArray<uint8>& PlatformDebugData, FName Format) const { }
 
 public:
 	/** Virtual destructor. */
diff --git a/Engine/Source/Runtime/RenderCore/Private/Shader.cpp b/Engine/Source/Runtime/RenderCore/Private/Shader.cpp
--- a/Engine/Source/Runtime/RenderCore/Private/Shader.cpp
+++ b/Engine/Source/Runtime/RenderCore/Private/Shader.cpp
@@ -735,7 +735,7 @@ void FShaderResource::SerializePlatformDebugData(FArchive& Ar)
 				const IShaderFormat* ShaderFormat = GetTargetPlatformManagerRef().FindShaderFormat(FormatName);
 				if (ShaderFormat)
 				{
-					ShaderFormat->NotifyShaderCooked(PlatformDebugData);
+					ShaderFormat->NotifyShaderCooked(PlatformDebugData, FormatName);
 				}
 			}
 		}
