Add sparse texture platform feature

diff --git a/Engine/Source/Developer/TargetPlatform/Public/Interfaces/ITargetPlatform.h b/Engine/Source/Developer/TargetPlatform/Public/Interfaces/ITargetPlatform.h
--- a/Engine/Source/Developer/TargetPlatform/Public/Interfaces/ITargetPlatform.h
+++ b/Engine/Source/Developer/TargetPlatform/Public/Interfaces/ITargetPlatform.h
@@ -84,6 +84,9 @@ enum class ETargetPlatformFeatures
 
 	/* The platform supports memory mapped animation */
 	MemoryMappedAnimation,
+
+	/* The platform supports sparse textures */
+	SparseTextures,
 };
 
 
