- Added early check before creating a TargetPlatform object if the PlatformInfo will be found (check a static function to see if it's usable before making it)
- Changed each platform's GetTargetPlatform[s] function

diff --git a/Engine/Source/Developer/AllDesktopTargetPlatform/Private/AllDesktopTargetPlatformModule.cpp b/Engine/Source/Developer/AllDesktopTargetPlatform/Private/AllDesktopTargetPlatformModule.cpp
--- a/Engine/Source/Developer/AllDesktopTargetPlatform/Private/AllDesktopTargetPlatformModule.cpp
+++ b/Engine/Source/Developer/AllDesktopTargetPlatform/Private/AllDesktopTargetPlatformModule.cpp
@@ -33,7 +33,7 @@ public:
 
 	virtual ITargetPlatform* GetTargetPlatform() override
 	{
-		if (Singleton == NULL)
+		if (Singleton == NULL && FAllDesktopTargetPlatform::IsUsable())
 		{
 			Singleton = new FAllDesktopTargetPlatform();
 		}
diff --git a/Engine/Source/Developer/Android/AndroidTargetPlatform/Private/AndroidTargetPlatformModule.cpp b/Engine/Source/Developer/Android/AndroidTargetPlatform/Private/AndroidTargetPlatformModule.cpp
--- a/Engine/Source/Developer/Android/AndroidTargetPlatform/Private/AndroidTargetPlatformModule.cpp
+++ b/Engine/Source/Developer/Android/AndroidTargetPlatform/Private/AndroidTargetPlatformModule.cpp
@@ -40,7 +40,7 @@ public:
 
 	virtual TArray<ITargetPlatform*> GetTargetPlatforms() override
 	{
-		if (TargetPlatforms.Num() == 0)
+		if (TargetPlatforms.Num() == 0 && FAndroidTargetPlatform::IsUsable())
 		{
 			for (int32 Type = 0; Type < 2; Type++)
 			{
diff --git a/Engine/Source/Developer/Android/Android_ETC1aTargetPlatform/Private/Android_ETC1aTargetPlatformModule.cpp b/Engine/Source/Developer/Android/Android_ETC1aTargetPlatform/Private/Android_ETC1aTargetPlatformModule.cpp
--- a/Engine/Source/Developer/Android/Android_ETC1aTargetPlatform/Private/Android_ETC1aTargetPlatformModule.cpp
+++ b/Engine/Source/Developer/Android/Android_ETC1aTargetPlatform/Private/Android_ETC1aTargetPlatformModule.cpp
@@ -45,7 +45,7 @@ public:
 
 	virtual ITargetPlatform* GetTargetPlatform() override
 	{
-		if (AndroidTargetSingleton == NULL)
+		if (AndroidTargetSingleton == NULL && FAndroid_ETC1aTargetPlatform::IsUsable())
 		{
 			AndroidTargetSingleton = new FAndroid_ETC1aTargetPlatform();
 		}
diff --git a/Engine/Source/Developer/HTML5/HTML5TargetPlatform/Private/HTML5TargetPlatformModule.cpp b/Engine/Source/Developer/HTML5/HTML5TargetPlatform/Private/HTML5TargetPlatformModule.cpp
--- a/Engine/Source/Developer/HTML5/HTML5TargetPlatform/Private/HTML5TargetPlatformModule.cpp
+++ b/Engine/Source/Developer/HTML5/HTML5TargetPlatform/Private/HTML5TargetPlatformModule.cpp
@@ -36,7 +36,7 @@ public:
 
 	virtual ITargetPlatform* GetTargetPlatform( ) override
 	{
-		if (Singleton == NULL)
+		if (Singleton == NULL && FHTML5TargetPlatform::IsUsable())
 		{
 			// finally, make the interface object
 			Singleton = new FHTML5TargetPlatform();
diff --git a/Engine/Source/Developer/IOS/IOSTargetPlatform/Private/IOSTargetPlatformModule.cpp b/Engine/Source/Developer/IOS/IOSTargetPlatform/Private/IOSTargetPlatformModule.cpp
--- a/Engine/Source/Developer/IOS/IOSTargetPlatform/Private/IOSTargetPlatformModule.cpp
+++ b/Engine/Source/Developer/IOS/IOSTargetPlatform/Private/IOSTargetPlatformModule.cpp
@@ -26,7 +26,7 @@ public:
 
 	virtual TArray<ITargetPlatform*> GetTargetPlatforms() override
 	{
-		if (TargetPlatforms.Num() == 0)
+		if (TargetPlatforms.Num() == 0 && FIOSTargetPlatform::IsUsable())
 		{
 			TargetPlatforms.Add(new FIOSTargetPlatform(false, true));
 			TargetPlatforms.Add(new FIOSTargetPlatform(false, false));
diff --git a/Engine/Source/Developer/IOS/TVOSTargetPlatform/Private/TVOSTargetPlatformModule.cpp b/Engine/Source/Developer/IOS/TVOSTargetPlatform/Private/TVOSTargetPlatformModule.cpp
--- a/Engine/Source/Developer/IOS/TVOSTargetPlatform/Private/TVOSTargetPlatformModule.cpp
+++ b/Engine/Source/Developer/IOS/TVOSTargetPlatform/Private/TVOSTargetPlatformModule.cpp
@@ -28,7 +28,7 @@ public:
 
 	virtual TArray<ITargetPlatform*> GetTargetPlatforms() override
 	{
-		if (TargetPlatforms.Num() == 0)
+		if (TargetPlatforms.Num() == 0 && FIOSTargetPlatform::IsUsable())
 		{
 			// add Game and Client TPs
 			TargetPlatforms.Add(new FIOSTargetPlatform(true, true));
diff --git a/Engine/Source/Developer/Linux/LinuxClientTargetPlatform/Private/LinuxClientTargetPlatformModule.cpp b/Engine/Source/Developer/Linux/LinuxClientTargetPlatform/Private/LinuxClientTargetPlatformModule.cpp
--- a/Engine/Source/Developer/Linux/LinuxClientTargetPlatform/Private/LinuxClientTargetPlatformModule.cpp
+++ b/Engine/Source/Developer/Linux/LinuxClientTargetPlatform/Private/LinuxClientTargetPlatformModule.cpp
@@ -34,7 +34,7 @@ public:
 
 	virtual ITargetPlatform* GetTargetPlatform( )
 	{
-		if (Singleton == NULL)
+		if (Singleton == NULL && TLinuxTargetPlatform<FLinuxPlatformProperties<false, false, true> >::IsUsable())
 		{
  			Singleton = new TLinuxTargetPlatform<FLinuxPlatformProperties<false, false, true> >();
 		}
diff --git a/Engine/Source/Developer/Linux/LinuxNoEditorTargetPlatform/Private/LinuxNoEditorTargetPlatformModule.cpp b/Engine/Source/Developer/Linux/LinuxNoEditorTargetPlatform/Private/LinuxNoEditorTargetPlatformModule.cpp
--- a/Engine/Source/Developer/Linux/LinuxNoEditorTargetPlatform/Private/LinuxNoEditorTargetPlatformModule.cpp
+++ b/Engine/Source/Developer/Linux/LinuxNoEditorTargetPlatform/Private/LinuxNoEditorTargetPlatformModule.cpp
@@ -33,7 +33,7 @@ public:
 
 	virtual ITargetPlatform* GetTargetPlatform( )
 	{
-		if (Singleton == NULL)
+		if (Singleton == NULL && TLinuxTargetPlatform<FLinuxPlatformProperties<false, false, false> >::IsUsable())
 		{
 			Singleton = new TLinuxTargetPlatform<FLinuxPlatformProperties<false, false, false> >();
 		}
diff --git a/Engine/Source/Developer/Linux/LinuxServerTargetPlatform/Private/LinuxServerTargetPlatformModule.cpp b/Engine/Source/Developer/Linux/LinuxServerTargetPlatform/Private/LinuxServerTargetPlatformModule.cpp
--- a/Engine/Source/Developer/Linux/LinuxServerTargetPlatform/Private/LinuxServerTargetPlatformModule.cpp
+++ b/Engine/Source/Developer/Linux/LinuxServerTargetPlatform/Private/LinuxServerTargetPlatformModule.cpp
@@ -33,7 +33,7 @@ public:
 
 	virtual ITargetPlatform* GetTargetPlatform( )
 	{
-		if (Singleton == NULL)
+		if (Singleton == NULL && TLinuxTargetPlatform<FLinuxPlatformProperties<false, false, false> >::IsUsable())
 		{
 			Singleton = new TLinuxTargetPlatform<FLinuxPlatformProperties<false, true, false> >();
 		}
diff --git a/Engine/Source/Developer/Linux/LinuxTargetPlatform/Private/LinuxTargetPlatformModule.cpp b/Engine/Source/Developer/Linux/LinuxTargetPlatform/Private/LinuxTargetPlatformModule.cpp
--- a/Engine/Source/Developer/Linux/LinuxTargetPlatform/Private/LinuxTargetPlatformModule.cpp
+++ b/Engine/Source/Developer/Linux/LinuxTargetPlatform/Private/LinuxTargetPlatformModule.cpp
@@ -43,7 +43,7 @@ public:
 
 	virtual ITargetPlatform* GetTargetPlatform( ) override
 	{
-		if (Singleton == NULL)
+		if (Singleton == NULL && TLinuxTargetPlatform<FLinuxPlatformProperties<true, false, false> >::IsUsable())
 		{
 			Singleton = new TLinuxTargetPlatform<FLinuxPlatformProperties<true, false, false> >();
 		}
diff --git a/Engine/Source/Developer/Lumin/LuminTargetPlatform/Private/LuminTargetPlatformModule.cpp b/Engine/Source/Developer/Lumin/LuminTargetPlatform/Private/LuminTargetPlatformModule.cpp
--- a/Engine/Source/Developer/Lumin/LuminTargetPlatform/Private/LuminTargetPlatformModule.cpp
+++ b/Engine/Source/Developer/Lumin/LuminTargetPlatform/Private/LuminTargetPlatformModule.cpp
@@ -32,7 +32,7 @@ public:
 
 	virtual TArray<ITargetPlatform*> GetTargetPlatforms() override
 	{
-		if (TargetPlatforms.Num() == 0)
+		if (TargetPlatforms.Num() == 0 && FLuminTargetPlatform::IsUsable())
 		{
 			TargetPlatforms.Add(new FLuminTargetPlatform(false));
 			TargetPlatforms.Add(new FLuminTargetPlatform(true));
diff --git a/Engine/Source/Developer/Mac/MacClientTargetPlatform/Private/MacClientTargetPlatformModule.cpp b/Engine/Source/Developer/Mac/MacClientTargetPlatform/Private/MacClientTargetPlatformModule.cpp
--- a/Engine/Source/Developer/Mac/MacClientTargetPlatform/Private/MacClientTargetPlatformModule.cpp
+++ b/Engine/Source/Developer/Mac/MacClientTargetPlatform/Private/MacClientTargetPlatformModule.cpp
@@ -30,7 +30,7 @@ public:
 
 	virtual ITargetPlatform* GetTargetPlatform( ) override
 	{
-		if (Singleton == NULL)
+		if (Singleton == NULL && TGenericMacTargetPlatform<false, false, true>::IsUsable())
 		{
 			Singleton = new TGenericMacTargetPlatform<false, false, true>();
 		}
diff --git a/Engine/Source/Developer/Mac/MacNoEditorTargetPlatform/Private/MacNoEditorTargetPlatformModule.cpp b/Engine/Source/Developer/Mac/MacNoEditorTargetPlatform/Private/MacNoEditorTargetPlatformModule.cpp
--- a/Engine/Source/Developer/Mac/MacNoEditorTargetPlatform/Private/MacNoEditorTargetPlatformModule.cpp
+++ b/Engine/Source/Developer/Mac/MacNoEditorTargetPlatform/Private/MacNoEditorTargetPlatformModule.cpp
@@ -31,7 +31,7 @@ public:
 
 	virtual ITargetPlatform* GetTargetPlatform( ) override
 	{
-		if (Singleton == NULL)
+		if (Singleton == NULL && TGenericMacTargetPlatform<false, false, false>::IsUsable())
 		{
 			Singleton = new TGenericMacTargetPlatform<false, false, false>();
 		}
diff --git a/Engine/Source/Developer/Mac/MacServerTargetPlatform/Private/MacServerTargetPlatformModule.cpp b/Engine/Source/Developer/Mac/MacServerTargetPlatform/Private/MacServerTargetPlatformModule.cpp
--- a/Engine/Source/Developer/Mac/MacServerTargetPlatform/Private/MacServerTargetPlatformModule.cpp
+++ b/Engine/Source/Developer/Mac/MacServerTargetPlatform/Private/MacServerTargetPlatformModule.cpp
@@ -30,7 +30,7 @@ public:
 
 	virtual ITargetPlatform* GetTargetPlatform( ) override
 	{
-		if (Singleton == NULL)
+		if (Singleton == NULL&& TGenericMacTargetPlatform<false, true, false>::IsUsable())
 		{
 			Singleton = new TGenericMacTargetPlatform<false, true, false>();
 		}
diff --git a/Engine/Source/Developer/Mac/MacTargetPlatform/Private/MacTargetPlatformModule.cpp b/Engine/Source/Developer/Mac/MacTargetPlatform/Private/MacTargetPlatformModule.cpp
--- a/Engine/Source/Developer/Mac/MacTargetPlatform/Private/MacTargetPlatformModule.cpp
+++ b/Engine/Source/Developer/Mac/MacTargetPlatform/Private/MacTargetPlatformModule.cpp
@@ -42,7 +42,7 @@ public:
 
 	virtual ITargetPlatform* GetTargetPlatform() override
 	{
-		if (Singleton == NULL)
+		if (Singleton == NULL && TGenericMacTargetPlatform<true, false, false>::IsUsable())
 		{
 			Singleton = new TGenericMacTargetPlatform<true, false, false>();
 		}
diff --git a/Engine/Source/Developer/TargetPlatform/Public/Common/TargetPlatformBase.h b/Engine/Source/Developer/TargetPlatform/Public/Common/TargetPlatformBase.h
--- a/Engine/Source/Developer/TargetPlatform/Public/Common/TargetPlatformBase.h
+++ b/Engine/Source/Developer/TargetPlatform/Public/Common/TargetPlatformBase.h
@@ -129,7 +129,7 @@ protected:
 	FTargetPlatformBase(const PlatformInfo::FPlatformInfo *const InPlatformInfo)
 		: PlatformInfo(InPlatformInfo)
 	{
-		check(PlatformInfo);
+		checkf(PlatformInfo, TEXT("Null PlatformInfo was passed to FTargetPlatformBase. Check the static IsUsable function before creating this object. See FWindowsTargetPlatformModule::GetTargetPlatform()"));
 
 		PlatformOrdinal = AssignPlatformOrdinal(*this);
 	}
@@ -151,6 +151,15 @@ class TTargetPlatformBase
 {
 public:
 
+	/**
+	 * Returns true if the target platform will be able to be  initialized with an FPlatformInfo. Because FPlatformInfo now comes from a .ini file,
+	 * it's possible that the .dll exists, but the .ini does not (should be uncommon, but is necessary to be handled)
+	 */
+	static bool IsUsable()
+	{
+		return PlatformInfo::FindPlatformInfo(TPlatformProperties::PlatformName()) != nullptr;
+	}
+	
 	/** Default constructor. */
 	TTargetPlatformBase()
 		: FTargetPlatformBase( PlatformInfo::FindPlatformInfo(TPlatformProperties::PlatformName()) )
diff --git a/Engine/Source/Developer/Windows/WindowsClientTargetPlatform/Private/WindowsClientTargetPlatformModule.cpp b/Engine/Source/Developer/Windows/WindowsClientTargetPlatform/Private/WindowsClientTargetPlatformModule.cpp
--- a/Engine/Source/Developer/Windows/WindowsClientTargetPlatform/Private/WindowsClientTargetPlatformModule.cpp
+++ b/Engine/Source/Developer/Windows/WindowsClientTargetPlatform/Private/WindowsClientTargetPlatformModule.cpp
@@ -25,7 +25,7 @@ public:
 
 	virtual ITargetPlatform* GetTargetPlatform( )
 	{
-		if (Singleton == nullptr)
+		if (Singleton == nullptr && TGenericWindowsTargetPlatform<false, false, true>::IsUsable())
 		{
 			Singleton = new TGenericWindowsTargetPlatform<false, false, true>();
 		}
diff --git a/Engine/Source/Developer/Windows/WindowsNoEditorTargetPlatform/Private/WindowsNoEditorTargetPlatformModule.cpp b/Engine/Source/Developer/Windows/WindowsNoEditorTargetPlatform/Private/WindowsNoEditorTargetPlatformModule.cpp
--- a/Engine/Source/Developer/Windows/WindowsNoEditorTargetPlatform/Private/WindowsNoEditorTargetPlatformModule.cpp
+++ b/Engine/Source/Developer/Windows/WindowsNoEditorTargetPlatform/Private/WindowsNoEditorTargetPlatformModule.cpp
@@ -25,7 +25,7 @@ public:
 
 	virtual ITargetPlatform* GetTargetPlatform( )
 	{
-		if (Singleton == nullptr)
+		if (Singleton == nullptr && TGenericWindowsTargetPlatform<false, false, false>::IsUsable())
 		{
 			Singleton = new TGenericWindowsTargetPlatform<false, false, false>();
 		}
diff --git a/Engine/Source/Developer/Windows/WindowsServerTargetPlatform/Private/WindowsServerTargetPlatformModule.cpp b/Engine/Source/Developer/Windows/WindowsServerTargetPlatform/Private/WindowsServerTargetPlatformModule.cpp
--- a/Engine/Source/Developer/Windows/WindowsServerTargetPlatform/Private/WindowsServerTargetPlatformModule.cpp
+++ b/Engine/Source/Developer/Windows/WindowsServerTargetPlatform/Private/WindowsServerTargetPlatformModule.cpp
@@ -25,7 +25,7 @@ public:
 
 	virtual ITargetPlatform* GetTargetPlatform( )
 	{
-		if (Singleton == nullptr)
+		if (Singleton == nullptr && TGenericWindowsTargetPlatform<false, true, false>::IsUsable())
 		{
 			Singleton = new TGenericWindowsTargetPlatform<false, true, false>();
 		}
diff --git a/Engine/Source/Developer/Windows/WindowsTargetPlatform/Private/WindowsTargetPlatformModule.cpp b/Engine/Source/Developer/Windows/WindowsTargetPlatform/Private/WindowsTargetPlatformModule.cpp
--- a/Engine/Source/Developer/Windows/WindowsTargetPlatform/Private/WindowsTargetPlatformModule.cpp
+++ b/Engine/Source/Developer/Windows/WindowsTargetPlatform/Private/WindowsTargetPlatformModule.cpp
@@ -49,7 +49,7 @@ public:
 
 	virtual ITargetPlatform* GetTargetPlatform( ) override
 	{
-		if (Singleton == nullptr)
+		if (Singleton == nullptr && TGenericWindowsTargetPlatform<true, false, false>::IsUsable())
 		{
 			Singleton = new TGenericWindowsTargetPlatform<true, false, false>();
 		}
