- Removed the wildcard searching for *TargetPlatform modules. Instead only attempt to load TPs where the PlatformInfo exists (ie, you have synced support for the platform)

diff --git a/Engine/Source/Developer/TargetPlatform/Private/TargetPlatformManagerModule.cpp b/Engine/Source/Developer/TargetPlatform/Private/TargetPlatformManagerModule.cpp
--- a/Engine/Source/Developer/TargetPlatform/Private/TargetPlatformManagerModule.cpp
+++ b/Engine/Source/Developer/TargetPlatform/Private/TargetPlatformManagerModule.cpp
@@ -613,28 +613,14 @@ protected:
 	/** Discovers the available target platforms. */
 	void DiscoverAvailablePlatforms()
 	{
-		DECLARE_SCOPE_CYCLE_COUNTER( TEXT( "FTargetPlatformManagerModule::DiscoverAvailablePlatforms" ), STAT_FTargetPlatformManagerModule_DiscoverAvailablePlatforms, STATGROUP_TargetPlatform );
+		DECLARE_SCOPE_CYCLE_COUNTER(TEXT("FTargetPlatformManagerModule::DiscoverAvailablePlatforms"), STAT_FTargetPlatformManagerModule_DiscoverAvailablePlatforms, STATGROUP_TargetPlatform);
 
 		Platforms.Empty(Platforms.Num());
 
-		TArray<FName> Modules;
-
-		FString ModuleWildCard = TEXT("*TargetPlatform");
-
-#if WITH_EDITOR
-		// if we have the editor and we are using -game
-		// only need to instantiate the current platform 
-#if PLATFORM_WINDOWS
-		if (IsRunningGame())
-		{
-			ModuleWildCard = TEXT("Windows*TargetPlatform");
-		}
-#endif
-#endif
-		FModuleManager::Get().FindModules(*ModuleWildCard, Modules);
-
 #if !IS_MONOLITHIC
 		// Find all module subdirectories and add them so we can load dependent modules for target platform modules
+		// We may not be able to restrict this to subdirectories found in FPlatformInfo because we could have a subdirectory
+		// that is not one of these platforms. Imagine a "Sega" shared directory for the "Genesis" and "Dreamcast" platforms
 		TArray<FString> ModuleSubdirs;
 		IFileManager::Get().FindFilesRecursive(ModuleSubdirs, *FPlatformProcess::GetModulesDirectory(), TEXT("*"), false, true);
 		for (const FString& ModuleSubdir : ModuleSubdirs)
@@ -643,50 +629,95 @@ protected:
 		}
 #endif
 
-		// remove this module from the list
-		Modules.Remove(FName(TEXT("TargetPlatform")));
-
-		if (!Modules.Num())
-		{
-			UE_LOG(LogTargetPlatformManager, Error, TEXT("No target platforms found!"));
-		}
+		// find a set of valid target platform names (the platform DataDrivenPlatformInfo.ini file was found indicates support for the platform 
+		// exists on disk, so the TP is expected to work)
+		const TArray<PlatformInfo::FPlatformInfo>& PlatformInfos = PlatformInfo::GetPlatformInfoArray();
 
-		FScopedSlowTask SlowTask(Modules.Num());
-		for (int32 Index = 0; Index < Modules.Num(); Index++)
+		TSet<ITargetPlatformModule*> ProcessedModules;
+		FScopedSlowTask SlowTask(PlatformInfos.Num());
+		for (const PlatformInfo::FPlatformInfo& PlatInfo : PlatformInfos)
 		{
 			SlowTask.EnterProgressFrame(1);
-			ITargetPlatformModule* Module = FModuleManager::LoadModulePtr<ITargetPlatformModule>(Modules[Index]);
-			if (Module)
+
+			// by defaulty load all PIs we can have
+			bool bLoadTargetPlatform = true;
+
+			// disabled?
+			if (PlatInfo.bEnabledForUse == false)
 			{
-				TArray<ITargetPlatform*> TargetPlatforms = Module->GetTargetPlatforms();
-				for (ITargetPlatform* Platform : TargetPlatforms) 
+				bLoadTargetPlatform = false;
+			}
+
+#if WITH_EDITOR
+			// if we have the editor and we are using -game
+			// only need to instantiate the current platform 
+			if (IsRunningGame())
+			{
+				if (PlatInfo.IniPlatformName == FPlatformProperties::IniPlatformName())
 				{
-					// would like to move this check to GetActiveTargetPlatforms, but too many things cache this result
-					// this setup will become faster after TTP 341897 is complete.
-RETRY_SETUPANDVALIDATE:
-					if (SetupAndValidateAutoSDK(Platform->GetPlatformInfo().AutoSDKPath))
-					{
-						UE_LOG(LogTemp, Display, TEXT("Module '%s' loaded TargetPlatform '%s'"), *Modules[Index].ToString(), *Platform->PlatformName());
-						Platforms.Add(Platform);
-					}
-					else
+					bLoadTargetPlatform = false;
+				}
+			}
+#endif
+
+			// now load the TP module
+			if (bLoadTargetPlatform)
+			{
+				// there are two ways targetplatform modules are setup: a single DLL per TargetPlatform, or a DLL for the platform
+				// that returns multiple TargetPlatforms. we try single first, then full platform
+				FName FullPlatformModuleName = *(PlatInfo.IniPlatformName + TEXT("TargetPlatform"));
+				FName SingleTargetPlatformModuleName = *(PlatInfo.TargetPlatformName.ToString() + TEXT("TargetPlatform"));
+
+				ITargetPlatformModule* Module = FModuleManager::Get().ModuleExists(*SingleTargetPlatformModuleName.ToString()) ?
+					FModuleManager::LoadModulePtr<ITargetPlatformModule>(SingleTargetPlatformModuleName) :
+					FModuleManager::LoadModulePtr<ITargetPlatformModule>(FullPlatformModuleName);
+
+				// if we have already processed this module, we can skip it!
+				if (ProcessedModules.Contains(Module))
+				{
+					continue;
+				}
+
+				// original logic for module loading here
+				if (Module)
+				{
+					ProcessedModules.Add(Module);
+
+					TArray<ITargetPlatform*> TargetPlatforms = Module->GetTargetPlatforms();
+					for (ITargetPlatform* Platform : TargetPlatforms)
 					{
-						// this hack is here because if you try and setup and validate autosdk some times it will fail because shared files are in use by another child cooker
-						static bool bIsChildCooker = FParse::Param(FCommandLine::Get(), TEXT("cookchild"));
-						if (bIsChildCooker)
+						// would like to move this check to GetActiveTargetPlatforms, but too many things cache this result
+						// this setup will become faster after TTP 341897 is complete.
+					RETRY_SETUPANDVALIDATE:
+						if (SetupAndValidateAutoSDK(Platform->GetPlatformInfo().AutoSDKPath))
+						{
+							UE_LOG(LogTargetPlatformManager, Display, TEXT("Loaded TargetPlatform '%s'"), *Platform->PlatformName());
+							Platforms.Add(Platform);
+						}
+						else
 						{
-							static int Counter = 0;
-							++Counter;
-							if (Counter < 10)
+							// this hack is here because if you try and setup and validate autosdk some times it will fail because shared files are in use by another child cooker
+							static bool bIsChildCooker = FParse::Param(FCommandLine::Get(), TEXT("cookchild"));
+							if (bIsChildCooker)
 							{
-								goto RETRY_SETUPANDVALIDATE;
+								static int Counter = 0;
+								++Counter;
+								if (Counter < 10)
+								{
+									goto RETRY_SETUPANDVALIDATE;
+								}
 							}
+							UE_LOG(LogTargetPlatformManager, Display, TEXT("Failed to SetupAndValidateAutoSDK for platform '%s'"), *Platform->PlatformName());
 						}
-						UE_LOG(LogTemp, Display, TEXT("Module '%s' failed to SetupAndValidateAutoSDK for platform '%s'"), *Modules[Index].ToString(), *Platform->PlatformName());
 					}
 				}
 			}
 		}
+
+		if (!Platforms.Num())
+		{
+			UE_LOG(LogTargetPlatformManager, Error, TEXT("No target platforms found!"));
+		}
 	}
 
 	bool UpdatePlatformEnvironment(const FString& PlatformName, TArray<FString> &Keys, TArray<FString> &Values) override
