Fixed a check in TargetPlatform (failing to set running target platform when running with editor + "-game").

diff --git a/Engine/Source/Developer/TargetPlatform/Private/TargetPlatformManagerModule.cpp b/Engine/Source/Developer/TargetPlatform/Private/TargetPlatformManagerModule.cpp
--- a/Engine/Source/Developer/TargetPlatform/Private/TargetPlatformManagerModule.cpp
+++ b/Engine/Source/Developer/TargetPlatform/Private/TargetPlatformManagerModule.cpp
@@ -653,7 +653,7 @@ protected:
 			// only need to instantiate the current platform 
 			if (IsRunningGame())
 			{
-				if (PlatInfo.IniPlatformName == FPlatformProperties::IniPlatformName())
+				if (PlatInfo.IniPlatformName != FPlatformProperties::IniPlatformName())
 				{
 					bLoadTargetPlatform = false;
 				}
