Check if the platform's ini name is valid (not empty) before trying to load the target platform module.

If the platform's ini name is empty and it drops down to using [IniName]+TargetPlatform it will try to recursively load TargetPlatform.
This can happen in cases where the AllDesktopTargetPlatform module is not available.

diff --git a/Engine/Source/Developer/TargetPlatform/Private/TargetPlatformManagerModule.cpp b/Engine/Source/Developer/TargetPlatform/Private/TargetPlatformManagerModule.cpp
--- a/Engine/Source/Developer/TargetPlatform/Private/TargetPlatformManagerModule.cpp
+++ b/Engine/Source/Developer/TargetPlatform/Private/TargetPlatformManagerModule.cpp
@@ -694,6 +694,7 @@ protected:
 				// that returns multiple TargetPlatforms. we try single first, then full platform
 				FName FullPlatformModuleName = *(PlatInfo.IniPlatformName + TEXT("TargetPlatform"));
 				FName SingleTargetPlatformModuleName = *(PlatInfo.TargetPlatformName.ToString() + TEXT("TargetPlatform"));
+				bool bFullPlatformModuleNameIsValid = !PlatInfo.IniPlatformName.IsEmpty();
 
 				ITargetPlatformModule* Module = nullptr;
 				
@@ -701,7 +702,7 @@ protected:
 				{
 					Module = FModuleManager::LoadModulePtr<ITargetPlatformModule>(SingleTargetPlatformModuleName);
 				}
-				else if (FModuleManager::Get().ModuleExists(*FullPlatformModuleName.ToString()))
+				else if (bFullPlatformModuleNameIsValid && FModuleManager::Get().ModuleExists(*FullPlatformModuleName.ToString()))
 				{
 					Module = FModuleManager::LoadModulePtr<ITargetPlatformModule>(FullPlatformModuleName);
 				}
