Don't assume a data driven target platform's module will always exist from its full name. Stops some "module not found" warnings appearing for platforms that have data driven platform configs but don't actually build.

diff --git a/Engine/Source/Developer/TargetPlatform/Private/TargetPlatformManagerModule.cpp b/Engine/Source/Developer/TargetPlatform/Private/TargetPlatformManagerModule.cpp
--- a/Engine/Source/Developer/TargetPlatform/Private/TargetPlatformManagerModule.cpp
+++ b/Engine/Source/Developer/TargetPlatform/Private/TargetPlatformManagerModule.cpp
@@ -667,9 +667,16 @@ protected:
 				FName FullPlatformModuleName = *(PlatInfo.IniPlatformName + TEXT("TargetPlatform"));
 				FName SingleTargetPlatformModuleName = *(PlatInfo.TargetPlatformName.ToString() + TEXT("TargetPlatform"));
 
-				ITargetPlatformModule* Module = FModuleManager::Get().ModuleExists(*SingleTargetPlatformModuleName.ToString()) ?
-					FModuleManager::LoadModulePtr<ITargetPlatformModule>(SingleTargetPlatformModuleName) :
-					FModuleManager::LoadModulePtr<ITargetPlatformModule>(FullPlatformModuleName);
+				ITargetPlatformModule* Module = nullptr;
+				
+				if (FModuleManager::Get().ModuleExists(*SingleTargetPlatformModuleName.ToString()))
+				{
+					Module = FModuleManager::LoadModulePtr<ITargetPlatformModule>(SingleTargetPlatformModuleName);
+				}
+				else if (FModuleManager::Get().ModuleExists(*FullPlatformModuleName.ToString()))
+				{
+					Module = FModuleManager::LoadModulePtr<ITargetPlatformModule>(FullPlatformModuleName);
+				}
 
 				// if we have already processed this module, we can skip it!
 				if (ProcessedModules.Contains(Module))
