Handle an invalid platform passed to the cook commandlet without a fatal assert.

diff --git a/Engine/Source/Developer/TargetPlatform/Private/TargetPlatformManagerModule.cpp b/Engine/Source/Developer/TargetPlatform/Private/TargetPlatformManagerModule.cpp
--- a/Engine/Source/Developer/TargetPlatform/Private/TargetPlatformManagerModule.cpp
+++ b/Engine/Source/Developer/TargetPlatform/Private/TargetPlatformManagerModule.cpp
@@ -79,6 +79,7 @@ public:
 	FTargetPlatformManagerModule()
 		: bRestrictFormatsToRuntimeOnly(false)
 		, bForceCacheUpdate(true)
+		, bHasInitErrors(false)
 		, bIgnoreFirstDelegateCall(true)
 	{
 #if AUTOSDKS_ENABLED		
@@ -111,15 +112,9 @@ public:
 			}
 		}
 #endif
-
-		SetupSDKStatus();
-		//GetTargetPlatforms(); redudant with next call
-		GetActiveTargetPlatforms();
-		GetAudioFormats();
-		GetTextureFormats();
-		GetShaderFormats();
-
-		bForceCacheUpdate = false;
+		// Calling a virtual function from a constructor, but with no expectation that a derived implementation of this
+		// method would be called.  This is solely to avoid duplicating code in this implementation, not for polymorphism.
+		FTargetPlatformManagerModule::Invalidate();
 
 		FModuleManager::Get().OnModulesChanged().AddRaw(this, &FTargetPlatformManagerModule::ModulesChangesCallback);
 	}
@@ -134,6 +129,15 @@ public:
 
 	// ITargetPlatformManagerModule interface
 
+	virtual bool HasInitErrors(FString* OutErrorMessages) const
+	{
+		if (OutErrorMessages)
+		{
+			*OutErrorMessages = InitErrorMessages;
+		}
+		return bHasInitErrors;
+	}
+
 	virtual void Invalidate() override
 	{
 		bForceCacheUpdate = true;
@@ -141,9 +145,14 @@ public:
 		SetupSDKStatus();
 		//GetTargetPlatforms(); redudant with next call
 		GetActiveTargetPlatforms();
-		GetAudioFormats();
-		GetTextureFormats();
-		GetShaderFormats();
+
+		// If we've had an error due to an invalid target platform, don't do additional work
+		if (!bHasInitErrors)
+		{
+			GetAudioFormats();
+			GetTextureFormats();
+			GetShaderFormats();
+		}
 
 		bForceCacheUpdate = false;
 	}
@@ -234,6 +243,8 @@ public:
 		if (!bInitialized || bForceCacheUpdate)
 		{
 			bInitialized = true;
+			bHasInitErrors = false; // If we had errors before, reset the flag and see later in this function if there are errors.
+			InitErrorMessages.Empty();
 
 			Results.Empty(Results.Num());
 
@@ -276,9 +287,11 @@ public:
 					if (Results.Num() == 0)
 					{
 						// An invalid platform was specified...
-						// Inform the user and exit.
+						// Inform the user.
+						bHasInitErrors = true;
+						InitErrorMessages.Appendf(TEXT("Invalid target platform specified (%s). Available = { %s } "), *PlatformStr, *AvailablePlatforms);
 						UE_LOG(LogTargetPlatformManager, Error, TEXT("Invalid target platform specified (%s). Available = { %s } "), *PlatformStr, *AvailablePlatforms);
-						UE_LOG(LogTargetPlatformManager, Fatal, TEXT("Invalid target platform specified (%s). Available = { %s } "), *PlatformStr, *AvailablePlatforms);
+						return Results;
 					}
 				}
 			}
@@ -1109,6 +1122,8 @@ private:
 		SDKStatusMessage += Message;
 	}
 
+	FString InitErrorMessages;
+
 	// If true we should build formats that are actually required for use by the runtime. 
 	// This happens for an ordinary editor run and more specifically whenever there is no
 	// TargetPlatform= on the command line.
@@ -1118,6 +1133,9 @@ private:
 	// in case of a module reload of a TargetPlatform-Module.
 	bool bForceCacheUpdate;
 
+	// Flag to indicate that there were errors during initialization
+	bool bHasInitErrors;
+
 	// Flag to avoid redunant reloads
 	bool bIgnoreFirstDelegateCall;
 
diff --git a/Engine/Source/Developer/TargetPlatform/Public/Interfaces/ITargetPlatformManagerModule.h b/Engine/Source/Developer/TargetPlatform/Public/Interfaces/ITargetPlatformManagerModule.h
--- a/Engine/Source/Developer/TargetPlatform/Public/Interfaces/ITargetPlatformManagerModule.h
+++ b/Engine/Source/Developer/TargetPlatform/Public/Interfaces/ITargetPlatformManagerModule.h
@@ -131,6 +131,14 @@ public:
 	 * @return Collection of texture formats.
 	 */
 	virtual const TArray<const class ITextureFormat*>& GetTextureFormats() = 0;
+
+	/**
+	 * Determine if there were errors during the initialization of the platform manager.
+	 *
+	 * @param OutErrorMessages Optional pointer to an FString that will have the error messages appended to it.
+	 * @return True if there were errors during the initialization of the platform manager, False otherwise.
+	 */
+	virtual bool HasInitErrors(FString* OutErrorMessages) const = 0;
 	
 	/**
 	 * Invalidates the target platform module.
diff --git a/Engine/Source/Runtime/Core/Private/Misc/CoreMisc.cpp b/Engine/Source/Runtime/Core/Private/Misc/CoreMisc.cpp
--- a/Engine/Source/Runtime/Core/Private/Misc/CoreMisc.cpp
+++ b/Engine/Source/Runtime/Core/Private/Misc/CoreMisc.cpp
@@ -9,7 +9,9 @@
 #include "HAL/PlatformTime.h"
 #include "Misc/App.h"
 #include "Misc/LazySingleton.h"
+#include "Misc/OutputDeviceError.h"
 #include "Misc/ScopeLock.h"
+#include "CoreGlobals.h"
 
 /** For FConfigFile in appInit							*/
 #include "Misc/ConfigCacheIni.h"
@@ -117,7 +119,7 @@ class FDerivedDataCacheInterface& GetDerivedDataCacheRef()
 	return *SingletonInterface;
 }
 
-class ITargetPlatformManagerModule* GetTargetPlatformManager()
+class ITargetPlatformManagerModule* GetTargetPlatformManager(bool bFailOnInitErrors)
 {
 	static class ITargetPlatformManagerModule* SingletonInterface = NULL;
 	if (!FPlatformProperties::RequiresCookedData())
@@ -128,6 +130,12 @@ class ITargetPlatformManagerModule* GetTargetPlatformManager()
 			check(IsInGameThread());
 			bInitialized = true;
 			SingletonInterface = FModuleManager::LoadModulePtr<ITargetPlatformManagerModule>("TargetPlatform");
+
+			FString InitErrors;
+			if (bFailOnInitErrors && SingletonInterface && SingletonInterface->HasInitErrors(&InitErrors))
+			{
+				GError->Log(*InitErrors);
+			}
 		}
 	}
 	return SingletonInterface;
diff --git a/Engine/Source/Runtime/Core/Public/Misc/CoreMisc.h b/Engine/Source/Runtime/Core/Public/Misc/CoreMisc.h
--- a/Engine/Source/Runtime/Core/Public/Misc/CoreMisc.h
+++ b/Engine/Source/Runtime/Core/Public/Misc/CoreMisc.h
@@ -71,8 +71,13 @@ CORE_API class FDerivedDataCacheInterface* GetDerivedDataCache();
 /** Return the DDC interface, fatal error if it is not available. **/
 CORE_API class FDerivedDataCacheInterface& GetDerivedDataCacheRef();
 
-/** Return the Target Platform Manager interface, if it is available, otherwise return NULL **/
-CORE_API class ITargetPlatformManagerModule* GetTargetPlatformManager();
+/**
+ * Return the Target Platform Manager interface, if it is available, otherwise return nullptr.
+ *
+ * @param bFailOnInitErrors If true (default) and errors occur during init of the TPM, an error will be logged and the process may terminate, otherwise will return whether there was an error or not.
+ * @return The Target Platform Manager interface, if it is available, otherwise return nullptr.
+*/
+CORE_API class ITargetPlatformManagerModule* GetTargetPlatformManager(bool bFailOnInitErrors = true);
 
 /** Return the Target Platform Manager interface, fatal error if it is not available. **/
 CORE_API class ITargetPlatformManagerModule& GetTargetPlatformManagerRef();
diff --git a/Engine/Source/Runtime/Launch/Private/LaunchEngineLoop.cpp b/Engine/Source/Runtime/Launch/Private/LaunchEngineLoop.cpp
--- a/Engine/Source/Runtime/Launch/Private/LaunchEngineLoop.cpp
+++ b/Engine/Source/Runtime/Launch/Private/LaunchEngineLoop.cpp
@@ -2128,6 +2128,23 @@ int32 FEngineLoop::PreInitPreStartupScreen(const TCHAR* CmdLine)
 		InitializeStdOutDevice();
 	}
 
+	bool bIsCook = bHasCommandletToken && (Token == TEXT("cookcommandlet"));
+#if WITH_EDITOR
+	{
+		if (bIsCook)
+		{
+			// Target platform manager can only be initialized successfully after 
+			ITargetPlatformManagerModule* TargetPlatformManager = GetTargetPlatformManager(false);
+			FString InitErrors;
+			if (TargetPlatformManager && TargetPlatformManager->HasInitErrors(&InitErrors))
+			{
+				RequestEngineExit(InitErrors);
+				return 1;
+			}
+		}
+	}
+#endif
+
 	{
 		SCOPED_BOOT_TIMING("IPlatformFeaturesModule::Get()");
 		// allow the platform to start up any features it may need
@@ -2313,8 +2330,7 @@ int32 FEngineLoop::PreInitPreStartupScreen(const TCHAR* CmdLine)
 	}
 
 
-	FString Commandline = FCommandLine::Get();
-	bool bEnableShaderCompile = !FParse::Param(*Commandline, TEXT("NoShaderCompile"));
+	bool bEnableShaderCompile = !FParse::Param(FCommandLine::Get(), TEXT("NoShaderCompile"));
 
 	if (bEnableShaderCompile && !FPlatformProperties::RequiresCookedData())
 	{
@@ -2349,8 +2365,7 @@ int32 FEngineLoop::PreInitPreStartupScreen(const TCHAR* CmdLine)
 		// hack: don't load global shaders if we are cooking we will load the shaders for the correct platform later
 		if (bEnableShaderCompile &&
 			!IsRunningDedicatedServer() &&
-			Commandline.Contains(TEXT("cookcommandlet")) == false &&
-			Commandline.Contains(TEXT("run=cook")) == false)
+			!bIsCook)
 			// if (FParse::Param(FCommandLine::Get(), TEXT("Multiprocess")) == false)
 		{
 			LLM_SCOPE(ELLMTag::Shaders);
@@ -5095,7 +5110,7 @@ bool FEngineLoop::AppInit( )
 		// Find the editor target
 		FString EditorTargetFileName;
 		for (const FTargetInfo& Target : FDesktopPlatformModule::Get()->GetTargetsForProject(FPaths::GetProjectFilePath()))
-	{
+		{
 			if (Target.Type == EBuildTargetType::Editor)
 			{
 				// Read the editor target receipt
