Optimize FindTargetPlatform and its most expensive caller during cooking

This reduced the cost of the function by 91% while cooking both a large title and a small sample project. The wall time saved on the large title that was tested was 42.4 seconds.

diff --git a/Engine/Source/Developer/TargetPlatform/Private/TargetPlatformManagerModule.cpp b/Engine/Source/Developer/TargetPlatform/Private/TargetPlatformManagerModule.cpp
--- a/Engine/Source/Developer/TargetPlatform/Private/TargetPlatformManagerModule.cpp
+++ b/Engine/Source/Developer/TargetPlatform/Private/TargetPlatformManagerModule.cpp
@@ -179,16 +179,13 @@ public:
 		return nullptr;
 	}
 
-	virtual ITargetPlatform* FindTargetPlatform(const FString& Name) override
+	virtual ITargetPlatform* FindTargetPlatform(const FStringView& Name) override
 	{
-		const TArray<ITargetPlatform*>& TargetPlatforms = GetTargetPlatforms();	
-		
-		for (int32 Index = 0; Index < TargetPlatforms.Num(); Index++)
+		GetTargetPlatforms(); // Populates PlatformsByName
+
+		if (ITargetPlatform** Platform = PlatformsByName.Find(FName(Name)))
 		{
-			if (TargetPlatforms[Index]->PlatformName() == Name)
-			{
-				return TargetPlatforms[Index];
-			}
+			return *Platform;
 		}
 
 		return nullptr;
@@ -643,6 +640,7 @@ protected:
 		DECLARE_SCOPE_CYCLE_COUNTER(TEXT("FTargetPlatformManagerModule::DiscoverAvailablePlatforms"), STAT_FTargetPlatformManagerModule_DiscoverAvailablePlatforms, STATGROUP_TargetPlatform);
 
 		Platforms.Empty(Platforms.Num());
+		PlatformsByName.Empty(PlatformsByName.Num());
 
 #if !IS_MONOLITHIC
 		// Find all module subdirectories and add them so we can load dependent modules for target platform modules
@@ -725,8 +723,10 @@ protected:
 					RETRY_SETUPANDVALIDATE:
 						if (SetupAndValidateAutoSDK(Platform->GetPlatformInfo().AutoSDKPath))
 						{
-							UE_LOG(LogTargetPlatformManager, Display, TEXT("Loaded TargetPlatform '%s'"), *Platform->PlatformName());
+							const FString& PlatformName = Platform->PlatformName();
+							UE_LOG(LogTargetPlatformManager, Display, TEXT("Loaded TargetPlatform '%s'"), *PlatformName);
 							Platforms.Add(Platform);
+							PlatformsByName.Add(FName(PlatformName), Platform);
 						}
 						else
 						{
@@ -1142,6 +1142,9 @@ private:
 	// Holds the list of discovered platforms.
 	TArray<ITargetPlatform*> Platforms;
 
+	// Map for fast lookup of platforms by name.
+	TMap<FName, ITargetPlatform*> PlatformsByName;
+
 #if AUTOSDKS_ENABLED
 	// holds the list of Platforms that have attempted setup.
 	TMap<FName, bool> PlatformsSetup;
diff --git a/Engine/Source/Developer/TargetPlatform/Public/Interfaces/ITargetPlatformManagerModule.h b/Engine/Source/Developer/TargetPlatform/Public/Interfaces/ITargetPlatformManagerModule.h
--- a/Engine/Source/Developer/TargetPlatform/Public/Interfaces/ITargetPlatformManagerModule.h
+++ b/Engine/Source/Developer/TargetPlatform/Public/Interfaces/ITargetPlatformManagerModule.h
@@ -3,6 +3,7 @@
 #pragma once
 
 #include "CoreMinimal.h"
+#include "Containers/StringView.h"
 #include "Interfaces/ITargetDevice.h"
 #include "Modules/ModuleInterface.h"
 
@@ -63,7 +64,7 @@ public:
 	 * @param Name The target platform's short or long name.
 	 * @return The target platform, or nullptr if not found.
 	 */
-	virtual ITargetPlatform* FindTargetPlatform( const FString& Name ) = 0;
+	virtual ITargetPlatform* FindTargetPlatform( const FStringView& Name ) = 0;
 
 	/**
 	 * Finds a target platform by looking for one that supports a given value for a generic type of support
diff --git a/Engine/Source/Runtime/RenderCore/Private/ShaderCore.cpp b/Engine/Source/Runtime/RenderCore/Private/ShaderCore.cpp
--- a/Engine/Source/Runtime/RenderCore/Private/ShaderCore.cpp
+++ b/Engine/Source/Runtime/RenderCore/Private/ShaderCore.cpp
@@ -22,6 +22,7 @@
 #include "Interfaces/ITargetPlatform.h"
 #include "Interfaces/ITargetPlatformManagerModule.h"
 #include "Misc/ConfigCacheIni.h"
+#include "Misc/StringBuilder.h"
 #endif
 
 static TAutoConsoleVariable<int32> CVarShaderDevelopmentMode(
@@ -260,7 +261,9 @@ bool AllowDebugViewmodes(EShaderPlatform Platform)
 	const int32 ForceDebugViewValue = CVarForceDebugViewModes.GetValueOnAnyThread();
 	bool bForceEnable = ForceDebugViewValue == 1;
 	bool bForceDisable = ForceDebugViewValue == 2;
-	ITargetPlatform* TargetPlatform = GetTargetPlatformManager()->FindTargetPlatform(ShaderPlatformToPlatformName(Platform).ToString());
+	TStringBuilder<64> PlatformName;
+	ShaderPlatformToPlatformName(Platform).ToString(PlatformName);
+	ITargetPlatform* TargetPlatform = GetTargetPlatformManager()->FindTargetPlatform(PlatformName);
 	return (!bForceDisable) && (bForceEnable || !TargetPlatform || !TargetPlatform->RequiresCookedData());
 #else
 	return AllowDebugViewmodes();
