Skip "AllDesktop" platform properties as it will be removed soon and causes shader map keys to differ between game and editor modes.

#rnx

diff --git a/Engine/Source/Developer/TargetPlatform/Private/TargetPlatformManagerModule.cpp b/Engine/Source/Developer/TargetPlatform/Private/TargetPlatformManagerModule.cpp
--- a/Engine/Source/Developer/TargetPlatform/Private/TargetPlatformManagerModule.cpp
+++ b/Engine/Source/Developer/TargetPlatform/Private/TargetPlatformManagerModule.cpp
@@ -197,6 +197,14 @@ public:
 
 		for (int32 Index = 0; Index < TargetPlatforms.Num(); Index++)
 		{
+			//@todo-lh:
+			// FAllDesktopPlatformProperties will be removed soon as it's no longer maintained
+			// and will be replaced by the platform specific subclasses eventually, so skip "AllDesktop.
+			// Find platform specific subclass instead.
+			if (TargetPlatforms[Index]->PlatformName() == TEXT("AllDesktop"))
+			{
+				continue;
+			}
 			if (TargetPlatforms[Index]->SupportsValueForType(SupportType, RequiredSupportedValue))
 			{
 				return TargetPlatforms[Index];
