- Fixed -strict compiler error

diff --git a/Engine/Source/Developer/TargetPlatform/Public/Common/TargetPlatformBase.h b/Engine/Source/Developer/TargetPlatform/Public/Common/TargetPlatformBase.h
--- a/Engine/Source/Developer/TargetPlatform/Public/Common/TargetPlatformBase.h
+++ b/Engine/Source/Developer/TargetPlatform/Public/Common/TargetPlatformBase.h
@@ -408,7 +408,7 @@ public:
 	 * so we pass it up into the ctor
 	 */
 	TNonDesktopTargetPlatformBase(bool bInIsClientOnly, const TCHAR* CookFlavor=nullptr, const TCHAR* OverrideIniPlatformName=nullptr)
-		: TTargetPlatformBase(new PlatformInfo::FTargetPlatformInfo(
+		: TTargetPlatformBase<TPlatformProperties>(new PlatformInfo::FTargetPlatformInfo(
 			OverrideIniPlatformName ? FString(OverrideIniPlatformName) : FString(TPlatformProperties::IniPlatformName()),
 			bInIsClientOnly ? EBuildTargetType::Client : EBuildTargetType::Game,
 			CookFlavor))
