- Fixed non-unity error

diff --git a/Engine/Source/Developer/TargetPlatform/Public/Common/TargetPlatformBase.h b/Engine/Source/Developer/TargetPlatform/Public/Common/TargetPlatformBase.h
--- a/Engine/Source/Developer/TargetPlatform/Public/Common/TargetPlatformBase.h
+++ b/Engine/Source/Developer/TargetPlatform/Public/Common/TargetPlatformBase.h
@@ -5,6 +5,7 @@
 #include "CoreMinimal.h"
 #include "Interfaces/ITargetPlatform.h"
 #include "PlatformInfo.h"
+#include "Misc/ConfigCacheIni.h"
 
 /**
  * Base class for target platforms.
