- Fixed UFE not starting - will re-think TUrnkey module stuff so UFE can get the benefits

diff --git a/Engine/Source/Developer/TargetPlatform/Private/TargetPlatformManagerModule.cpp b/Engine/Source/Developer/TargetPlatform/Private/TargetPlatformManagerModule.cpp
--- a/Engine/Source/Developer/TargetPlatform/Private/TargetPlatformManagerModule.cpp
+++ b/Engine/Source/Developer/TargetPlatform/Private/TargetPlatformManagerModule.cpp
@@ -89,7 +89,9 @@ public:
 		, bHasInitErrors(false)
 		, bIgnoreFirstDelegateCall(true)
 	{
+#if WITH_EDITOR
 		ITurnkeySupportModule::Get().UpdateSdkInfo();
+#endif
 
 #if AUTOSDKS_ENABLED		
 		
