- Moved ProjectPackagingSettings to TargetPlatform module so eventually we can move Turnkey to a Developer module, and it can still use ProjectPackagingSettings (mostly read-only, doesn't need the editor customizations)

diff --git a/Engine/Config/BaseEngine.ini b/Engine/Config/BaseEngine.ini
--- a/Engine/Config/BaseEngine.ini
+++ b/Engine/Config/BaseEngine.ini
@@ -846,6 +846,10 @@ NearClipPlane=10.0
 +EnumRedirects=(OldName="/Script/Slate.ETextTransformPolicy", NewName="/Script/SlateCore.ETextTransformPolicy")
 +PropertyRedirects=(OldName="BlendProfile.BlendProfileMode", NewName="Mode")
 +PropertyRedirects=(OldName="RigVMInjectionInfo.StructNode", NewName="UnitNode")
++ClassRedirects=(OldName="/Script/UnrealEd.ProjectPackagingSettings", NewName="/Script/TargetPlatform.ProjectPackagingSettings")
++ClassRedirects=(OldName="/Script/UnrealEd.CookerSettings", NewName="/Script/TargetPlatform.CookerSettings")
++StructRedirects=(OldName="/Script/Engine.FilePath",NewName="/Script/CoreUObject.FilePath")
++StructRedirects=(OldName="/Script/Engine.DirectoryPath",NewName="/Script/CoreUObject.DirectoryPath")
 
 [CoreUObject.Metadata]
 ; Note: UnrealHeaderTool should be rerun after making changes to MetadataRedirects to catch any keys specified in class headers
diff --git a/Engine/Plugins/Editor/CryptoKeys/Source/CryptoKeys/CryptoKeys.Build.cs b/Engine/Plugins/Editor/CryptoKeys/Source/CryptoKeys/CryptoKeys.Build.cs
--- a/Engine/Plugins/Editor/CryptoKeys/Source/CryptoKeys/CryptoKeys.Build.cs
+++ b/Engine/Plugins/Editor/CryptoKeys/Source/CryptoKeys/CryptoKeys.Build.cs
@@ -26,7 +26,8 @@ public class CryptoKeys : ModuleRules
 				"CryptoKeysOpenSSL",
 				"Slate",
 				"SlateCore",
-				"GameProjectGeneration"
+				"GameProjectGeneration",
+				"TargetPlatform"
 		});
 	}
 }
\ No newline at end of file
diff --git a/Engine/Source/Developer/TargetPlatform/Classes/CookerSettings.h b/Engine/Source/Developer/TargetPlatform/Classes/CookerSettings.h
new file mode 100644
--- /dev/null
+++ b/Engine/Source/Developer/TargetPlatform/Classes/CookerSettings.h
@@ -0,0 +1,155 @@
+// Copyright Epic Games, Inc. All Rights Reserved.
+
+/*=============================================================================
+	CookerSettings.h: Declares the UCookerSettings class.
+=============================================================================*/
+
+#pragma once
+
+#include "CoreMinimal.h"
+#include "UObject/ObjectMacros.h"
+#include "Engine/DeveloperSettings.h"
+#include "CookerSettings.generated.h"
+
+struct FPropertyChangedEvent;
+
+UENUM()
+enum class ECookProgressDisplayMode : int32
+{
+	/** Don't display any progress messages */
+	Nothing = 0,
+
+	/** Display the number of remaining packages */
+	RemainingPackages = 1,
+
+	/** Display names of cooked packages */
+	PackageNames = 2,
+
+	/** Display the number of remaining packages and package names */
+	NamesAndRemainingPackages = 3,
+
+	/** MAX - invalid */
+	Max UMETA(Hidden),
+};
+
+UENUM()
+enum class EBlueprintComponentDataCookingMethod
+{
+	/** Do not generate optimized component data. No additional memory will be used. */
+	Disabled,
+
+	/** Generate optimized component data for all Blueprint types. This option will require the most additional memory. */
+	AllBlueprints,
+
+	/** Generate optimized component data only for Blueprint types that have explicitly enabled this feature in the class settings. */
+	EnabledBlueprintsOnly,
+};
+
+/**
+ * Various cooker settings.
+ */
+UCLASS(config=Engine, defaultconfig, meta=(DisplayName="Cooker"))
+class TARGETPLATFORM_API UCookerSettings : public UDeveloperSettings
+{
+	GENERATED_UCLASS_BODY()
+
+	virtual void PostInitProperties() override;
+#if WITH_EDITOR
+	virtual void PostEditChangeProperty(FPropertyChangedEvent& PropertyChangedEvent) override;
+#endif
+
+	/**
+	 * This class was moved from UnrealEd module, but to allow it to be used by developer tools, like UFE, it has moved to this module.
+	 * However, for back-compat, we want to use the old name in the ini files, so that everything works without needing to touch every
+	 * Game.ini file
+	 */
+	virtual void OverrideConfigSection(FString& InOutSectionName) override
+	{
+		InOutSectionName = TEXT("/Script/UnrealEd.CookerSettings");
+	}
+
+public:
+
+	UPROPERTY(GlobalConfig, EditAnywhere, Category = Cooker, meta = (DisplayName = "Enable cooking via network in the background of the editor, launch on uses this setting, requires device to have network access to editor", ConfigRestartRequired = true))
+	bool bEnableCookOnTheSide;
+
+	UPROPERTY(GlobalConfig, EditAnywhere, Category = Cooker, meta = (DisplayName = "Generate DDC data in background for desired launch on platform (speeds up launch on)"))
+	bool bEnableBuildDDCInBackground;
+
+	/** Enable -iterate for launch on */
+	UPROPERTY(GlobalConfig, EditAnywhere, Category = Cooker, meta = (DisplayName = "Iterative cooking for builds launched from the editor (launch on)"))
+	bool bIterativeCookingForLaunchOn;
+
+	/** Enable -iterate when triggering from File dropdown menu */
+	UPROPERTY(GlobalConfig, EditAnywhere, Category = Cooker, meta = (DisplayName = "Iterative cooking for the File->Cook Content menu item"))
+	bool bIterativeCookingForFileCookContent;
+
+	/** Enable -cookonthefly for launch on */
+	UPROPERTY(GlobalConfig, EditAnywhere, Category = Cooker, meta = (DisplayName = "Cooking on the fly when launching from the editor (launch on)"))
+	bool bCookOnTheFlyForLaunchOn;
+
+	UPROPERTY(GlobalConfig, EditAnywhere, Category = Cooker, meta = (
+		ConsoleVariable = "cook.displaymode", DisplayName = "Cooker Progress Display Mode",
+		ToolTip = "Controls log output of the cooker"))
+	ECookProgressDisplayMode CookProgressDisplayMode;
+
+	/** Ignore ini changes when doing iterative cooking, either in editor or out of editor */
+	UPROPERTY(GlobalConfig, EditAnywhere, Category = Cooker, AdvancedDisplay)
+	bool bIgnoreIniSettingsOutOfDateForIteration;
+
+	/** Ignore native header file changes when doing iterative cooking, either in editor or out of editor */
+	UPROPERTY(GlobalConfig, EditAnywhere, Category = Cooker, AdvancedDisplay)
+	bool bIgnoreScriptPackagesOutOfDateForIteration;
+
+	/** Whether or not to compile Blueprints in development mode when cooking. */
+	UPROPERTY(GlobalConfig, EditAnywhere, Category = Cooker, AdvancedDisplay)
+	bool bCompileBlueprintsInDevelopmentMode;
+
+	/** Generate optimized component data to speed up Blueprint construction at runtime. This option can increase the overall Blueprint memory usage in a cooked build. Requires Event-Driven Loading (EDL), which is enabled by default. */
+	UPROPERTY(GlobalConfig, EditAnywhere, Category = Cooker, AdvancedDisplay, meta = (DisplayName = "Generate optimized Blueprint component data"))
+	EBlueprintComponentDataCookingMethod BlueprintComponentDataCookingMethod;
+
+	/** List of class names to exclude when cooking for dedicated server */
+	UPROPERTY(GlobalConfig, EditAnywhere, Category = Cooker, AdvancedDisplay, meta = (DisplayName = "Classes excluded when cooking for dedicated server"))
+	TArray<FString> ClassesExcludedOnDedicatedServer;
+
+	/** List of module names to exclude when cooking for dedicated server */
+	UPROPERTY(GlobalConfig, EditAnywhere, Category = Cooker, AdvancedDisplay, meta = (DisplayName = "Modules excluded when cooking for dedicated server"))
+	TArray<FString> ModulesExcludedOnDedicatedServer;
+
+	/** List of class names to exclude when cooking for dedicated client */
+	UPROPERTY(GlobalConfig, EditAnywhere, Category = Cooker, AdvancedDisplay, meta = (DisplayName = "Classes excluded when cooking for dedicated client"))
+	TArray<FString> ClassesExcludedOnDedicatedClient;
+
+	/** List of module names to exclude when cooking for dedicated client */
+	UPROPERTY(GlobalConfig, EditAnywhere, Category = Cooker, AdvancedDisplay, meta = (DisplayName = "Modules excluded when cooking for dedicated client"))
+	TArray<FString> ModulesExcludedOnDedicatedClient;
+
+	/** List of r values that need to be versioned */
+	UPROPERTY(GlobalConfig, EditAnywhere, Category = Cooker, AdvancedDisplay, meta = (DisplayName = "r values that need to be versioned"))
+	TArray<FString> VersionedIntRValues;
+
+	/** Quality of 0 means fastest, 4 means best quality */
+	UPROPERTY(GlobalConfig, EditAnywhere, Category = Textures, meta = (DisplayName = "PVRTC Compression Quality (0-4, 0 is fastest)"))
+	int32 DefaultPVRTCQuality;
+
+	/** Quality of 0 means fastest, 3 means best quality */
+	UPROPERTY(GlobalConfig, EditAnywhere, Category = Textures, meta = (DisplayName = "ASTC Compression Quality vs Speed (0-3, 0 is fastest)"))
+	int32 DefaultASTCQualityBySpeed;
+
+	/** Quality of 0 means smallest (12x12 block size), 4 means best (4x4 block size) */
+	UPROPERTY(GlobalConfig, EditAnywhere, Category = Textures, meta = (DisplayName = "ASTC Compression Quality vs Size (0-4, 0 is smallest)"))
+	int32 DefaultASTCQualityBySize;
+
+	/** Allows opening cooked assets in the editor */
+	UPROPERTY(EditAnywhere, config, Category = Editor, meta = (
+		ConsoleVariable = "cook.AllowCookedDataInEditorBuilds", DisplayName = "Allow Cooked Content In The Editor",
+		ToolTip = "If true, the editor will be able to open cooked assets (limited to a subset of supported asset types).",
+		ConfigRestartRequired = true))
+	uint32 bAllowCookedDataInEditorBuilds : 1;
+
+private:
+	/** Deprecated. Use BlueprintComponentDataCookingMethod instead. */
+	UPROPERTY(GlobalConfig)
+	bool bCookBlueprintComponentTemplateData;
+};
diff --git a/Engine/Source/Developer/TargetPlatform/Classes/Settings/ProjectPackagingSettings.h b/Engine/Source/Developer/TargetPlatform/Classes/Settings/ProjectPackagingSettings.h
new file mode 100644
--- /dev/null
+++ b/Engine/Source/Developer/TargetPlatform/Classes/Settings/ProjectPackagingSettings.h
@@ -0,0 +1,518 @@
+// Copyright Epic Games, Inc. All Rights Reserved.
+
+#pragma once
+
+#include "CoreMinimal.h"
+#include "UObject/ObjectMacros.h"
+#include "UObject/Object.h"
+#include "ProjectPackagingSettings.generated.h"
+
+struct FTargetInfo;
+
+/**
+ * Enumerates the available build configurations for project packaging.
+ */
+UENUM()
+enum class EProjectPackagingBuildConfigurations : uint8
+{
+	/** Debug configuration. */
+	PPBC_Debug UMETA(DisplayName="Debug"),
+
+	/** DebugGame configuration. */
+	PPBC_DebugGame UMETA(DisplayName="DebugGame"),
+
+	/** Development configuration. */
+	PPBC_Development UMETA(DisplayName="Development"),
+
+	/** Test configuration. */
+	PPBC_Test UMETA(DisplayName="Test"),
+
+	/** Shipping configuration. */
+	PPBC_Shipping UMETA(DisplayName="Shipping"),
+
+	/** Number of entries in the enum. */
+	PPBC_MAX
+};
+
+/**
+ * Enumerates the available internationalization data presets for project packaging.
+ */
+UENUM()
+enum class EProjectPackagingInternationalizationPresets : uint8
+{
+	/** English only. */
+	English,
+
+	/** English, French, Italian, German, Spanish. */
+	EFIGS,
+
+	/** English, French, Italian, German, Spanish, Chinese, Japanese, Korean. */
+	EFIGSCJK,
+
+	/** Chinese, Japanese, Korean. */
+	CJK,
+
+	/** All known cultures. */
+	All
+};
+
+/**
+ * Determines whether to build the executable when packaging. Note the equivalence between these settings and EPlayOnBuildMode.
+ */
+UENUM()
+enum class EProjectPackagingBuild
+{
+	/** Always build. */
+	Always UMETA(DisplayName="Always"),
+
+	/** Never build. */
+	Never UMETA(DisplayName="Never"),
+
+	/** Default (if the Never build. */
+	IfProjectHasCode UMETA(DisplayName="If project has code, or running a locally built editor"),
+
+	/** If we're not packaging from a promoted build. */
+	IfEditorWasBuiltLocally UMETA(DisplayName="If running a locally built editor")
+};
+
+/**
+* Enumerates the available methods for Blueprint nativization during project packaging.
+*/
+UENUM()
+enum class EProjectPackagingBlueprintNativizationMethod : uint8
+{
+	/** Disable Blueprint nativization (default). */
+	Disabled,
+
+	/** Enable nativization for all Blueprint assets. */
+	Inclusive,
+
+	/** Enable nativization for selected Blueprint assets only. */
+	Exclusive
+};
+
+/**
+ * Implements the Editor's user settings.
+ */
+UCLASS(config=Game, defaultconfig)
+class TARGETPLATFORM_API UProjectPackagingSettings
+	: public UObject
+{
+	GENERATED_UCLASS_BODY()
+
+	/**
+	 * This class was moved from UnrealEd module, but to allow it to be used by developer tools, like UFE, it has moved to this module.
+	 * However, for back-compat, we want to use the old name in the ini files, so that everything works without needing to touch every
+	 * Game.ini file
+	 */
+	virtual void OverrideConfigSection(FString& InOutSectionName) override
+	{
+		InOutSectionName = TEXT("/Script/UnrealEd.ProjectPackagingSettings");
+	}
+
+public:
+
+	/**
+	 * Information about each packaging configuration
+	 */
+	struct FConfigurationInfo
+	{
+		EBuildConfiguration Configuration;
+		FText Name;
+		FText ToolTip;
+	};
+
+	/**
+	 * Static array of information about each configuration
+	 */
+	static const FConfigurationInfo ConfigurationInfo[(int)EProjectPackagingBuildConfigurations::PPBC_MAX];
+
+	/** Specifies whether to build the game executable during packaging. */
+	UPROPERTY(config, EditAnywhere, Category=Project)
+	EProjectPackagingBuild Build;
+
+	/** The build configuration for which the project is packaged. */
+	UPROPERTY(config, EditAnywhere, Category=Project)
+	EProjectPackagingBuildConfigurations BuildConfiguration;
+
+	/** Name of the target to build */
+	UPROPERTY(config, EditAnywhere, Category=Project)
+	FString BuildTarget;
+
+	/** The directory to which the packaged project will be copied. */
+	UPROPERTY(config, EditAnywhere, Category=Project)
+	FDirectoryPath StagingDirectory;
+
+	/**
+	 * If enabled, a full rebuild will be enforced each time the project is being packaged.
+	 * If disabled, only modified files will be built, which can improve iteration time.
+	 * Unless you iterate on packaging, we recommend full rebuilds when packaging.
+	 */
+	UPROPERTY(config, EditAnywhere, Category=Project)
+	bool FullRebuild;
+
+	/**
+	 * If enabled, a distribution build will be created and the shipping configuration will be used
+	 * If disabled, a development build will be created
+	 * Distribution builds are for publishing to the App Store
+	 */
+	UPROPERTY(config, EditAnywhere, Category=Project)
+	bool ForDistribution;
+
+	/** If enabled, debug files will be included in the packaged game */
+	UPROPERTY(config, EditAnywhere, Category=Project)
+	bool IncludeDebugFiles;
+
+	/** If enabled, then the project's Blueprint assets (including structs and enums) will be intermediately converted into C++ and used in the packaged project (in place of the .uasset files).*/
+	UPROPERTY(config, EditAnywhere, Category = Blueprints)
+	EProjectPackagingBlueprintNativizationMethod BlueprintNativizationMethod;
+
+	/** List of Blueprints to include for nativization when using the exclusive method. */
+	UPROPERTY(config, EditAnywhere, AdvancedDisplay, Category = Blueprints, meta = (DisplayName = "List of Blueprint assets to nativize", RelativeToGameContentDir, LongPackageName))
+	TArray<FFilePath> NativizeBlueprintAssets;
+
+	/** If enabled, the nativized assets code plugin will be added to the Visual Studio solution if it exists when regenerating the game project. Intended primarily to assist with debugging the target platform after cooking with nativization turned on. */
+	UPROPERTY(config, EditAnywhere, AdvancedDisplay, Category = Blueprints)
+	bool bIncludeNativizedAssetsInProjectGeneration;
+
+	/** Whether or not to exclude monolithic engine headers (e.g. Engine.h) in the generated code when nativizing Blueprint assets. This may improve C++ compiler performance if your game code does not depend on monolithic engine headers to build. */
+	UPROPERTY(config, EditAnywhere, AdvancedDisplay, Category = Blueprints)
+	bool bExcludeMonolithicEngineHeadersInNativizedCode;
+
+	/** If enabled, all content will be put into a one or more .pak files instead of many individual files (default = enabled). */
+	UPROPERTY(config, EditAnywhere, Category=Packaging)
+	bool UsePakFile;
+
+	/** If enabled, all packages will be put into one or more container files. */
+	UPROPERTY(config, EditAnywhere, Category = Packaging)
+	bool bUseIoStore;
+
+	/** 
+	 * If enabled, will generate pak file chunks.  Assets can be assigned to chunks in the editor or via a delegate (See ShooterGameDelegates.cpp). 
+	 * Can be used for streaming installs (PS4 Playgo, XboxOne Streaming Install, etc)
+	 */
+	UPROPERTY(config, EditAnywhere, Category=Packaging)
+	bool bGenerateChunks;
+
+	/** 
+	 * If enabled, no platform will generate chunks, regardless of settings in platform-specific ini files.
+	 */
+	UPROPERTY(config, EditAnywhere, Category=Packaging)
+	bool bGenerateNoChunks;
+
+	/**
+	 * Normally during chunk generation all dependencies of a package in a chunk will be pulled into that package's chunk.
+	 * If this is enabled then only hard dependencies are pulled in. Soft dependencies stay in their original chunk.
+	 */
+	UPROPERTY(config, EditAnywhere, Category = Packaging)
+	bool bChunkHardReferencesOnly;
+
+	/**
+	 * If true, individual files are only allowed to be in a single chunk and it will assign it to the lowest number requested
+	 * If false, it may end up in multiple chunks if requested by the cooker
+	 */
+	UPROPERTY(config, EditAnywhere, Category = Packaging, AdvancedDisplay)
+	bool bForceOneChunkPerFile;
+
+	/**
+	 * If > 0 this sets a maximum size per chunk. Chunks larger than this size will be split into multiple pak files such as pakchunk0_s1
+	 * This can be set in platform specific game.ini files
+	 */
+	UPROPERTY(config, EditAnywhere, Category = Packaging, AdvancedDisplay)
+	int64 MaxChunkSize;
+
+	/** 
+	 * If enabled, will generate data for HTTP Chunk Installer. This data can be hosted on webserver to be installed at runtime. Requires "Generate Chunks" enabled.
+	 */
+	UPROPERTY(config, EditAnywhere, Category=Packaging)
+	bool bBuildHttpChunkInstallData;
+
+	/** 
+	 * When "Build HTTP Chunk Install Data" is enabled this is the directory where the data will be build to.
+	 */	
+	UPROPERTY(config, EditAnywhere, Category = Packaging)
+	FDirectoryPath HttpChunkInstallDataDirectory;
+
+	/**
+	 * A comma separated list of formats to use for .pak file compression. If more than one is specified, the list is in order of priority, with fallbacks to other formats
+	 * in case of errors or unavailability of the format (plugin not enabled, etc).
+	 */
+	UPROPERTY(config, EditAnywhere, Category = Packaging, meta = (DisplayName = "Pak File Compression Format(s)"))
+	FString PakFileCompressionFormats;
+
+	/**
+	 * A generic setting for allowing a project to control compression settings during .pak file compression. For instance, if using the Oodle plugin, you could use -oodlemethod=Selkie -oodlelevel=Optimal1
+	 */
+	UPROPERTY(config, EditAnywhere, Category = Packaging, meta = (DisplayName = "Pak File Compression Commandline Options"))
+	FString PakFileAdditionalCompressionOptions;
+
+	/** 
+	 * Version name for HTTP Chunk Install Data.
+	 */
+	UPROPERTY(config, EditAnywhere, Category = Packaging)
+	FString HttpChunkInstallDataVersion;
+
+	/** Specifies whether to include an installer for prerequisites of packaged games, such as redistributable operating system components, on platforms that support it. */
+	UPROPERTY(config, EditAnywhere, Category=Prerequisites, meta=(DisplayName="Include prerequisites installer"))
+	bool IncludePrerequisites;
+
+	/** Specifies whether to include prerequisites alongside the game executable. */
+	UPROPERTY(config, EditAnywhere, Category = Prerequisites, meta = (DisplayName = "Include app-local prerequisites"))
+	bool IncludeAppLocalPrerequisites;
+
+	/** 
+	 * By default shader code gets saved inline inside material assets, 
+	 * enabling this option will store only shader code once as individual files
+	 * This will reduce overall package size but might increase loading time
+	 */
+	UPROPERTY(config, EditAnywhere, Category=Packaging)
+	bool bShareMaterialShaderCode;
+
+	/** 
+	 * With this option off, the shader code will be stored in the library essentially in a random order,
+	 * squarely the same in which the assets were loaded by the cooker. Enabling this will sort the shaders
+	 * by their hash, which makes the shader library more similar between the builds which can help patching, but
+	 * can adversely affect loading times.
+	 */
+	UPROPERTY(config, EditAnywhere, Category = Packaging, meta = (EditCondition = "bShareMaterialShaderCode"))
+	bool bDeterministicShaderCodeOrder;
+
+	/**
+	 * By default shader shader code gets saved into individual platform agnostic files,
+	 * enabling this option will use the platform-specific library format if and only if one is available
+	 * This will reduce overall package size but might increase loading time
+	 */
+	UPROPERTY(config, EditAnywhere, Category=Packaging, meta = (EditCondition = "bShareMaterialShaderCode", ConfigRestartRequired = true))
+	bool bSharedMaterialNativeLibraries;
+
+	/** A directory containing additional prerequisite packages that should be staged in the executable directory. Can be relative to $(EngineDir) or $(ProjectDir) */
+	UPROPERTY(config, EditAnywhere, Category=Prerequisites, AdvancedDisplay)
+	FDirectoryPath ApplocalPrerequisitesDirectory;
+
+	/**
+	 * Specifies whether to include the crash reporter in the packaged project. 
+	 * This is included by default for Blueprint based projects, but can optionally be disabled.
+	 */
+	UPROPERTY(config, EditAnywhere, Category=Packaging, AdvancedDisplay)
+	bool IncludeCrashReporter;
+
+	/** Predefined sets of culture whose internationalization data should be packaged. */
+	UPROPERTY(config, EditAnywhere, Category=Packaging, AdvancedDisplay, meta=(DisplayName="Internationalization Support"))
+	EProjectPackagingInternationalizationPresets InternationalizationPreset;
+
+	/** Cultures whose data should be cooked, staged, and packaged. */
+	UPROPERTY(config, EditAnywhere, Category=Packaging, AdvancedDisplay, meta=(DisplayName="Localizations to Package"))
+	TArray<FString> CulturesToStage;
+
+	/** List of localization targets that should be chunked during cooking (if using chunks) */
+	UPROPERTY(config, EditAnywhere, Category=Packaging, AdvancedDisplay)
+	TArray<FString> LocalizationTargetsToChunk;
+
+	/** The chunk ID that should be used as the catch-all chunk for any non-asset localized strings */
+	UPROPERTY(config, EditAnywhere, Category=Packaging, AdvancedDisplay)
+	int32 LocalizationTargetCatchAllChunkId = 0;
+
+	/**
+	 * Cook all things in the project content directory
+	 */
+	UPROPERTY(config, EditAnywhere, Category=Packaging, AdvancedDisplay, meta=(DisplayName="Cook everything in the project content directory (ignore list of maps below)"))
+	bool bCookAll;
+
+	/**
+	 * Cook only maps (this only affects the cookall flag)
+	 */
+	UPROPERTY(config, EditAnywhere, Category=Packaging, AdvancedDisplay, meta=(DisplayName="Cook only maps (this only affects cookall)"))
+	bool bCookMapsOnly;
+
+
+	/**
+	 * Create compressed cooked packages (decreased deployment size)
+	 */
+	UPROPERTY(config, EditAnywhere, Category=Packaging, AdvancedDisplay, meta=(DisplayName="Create compressed cooked packages"))
+	bool bCompressed;
+
+	/**
+	* Encrypt ini files inside of the pak file
+	* NOTE: Replaced by the settings inside the cryptokeys system. Kept here for legacy migration purposes.
+	*/
+	UPROPERTY(config)
+	bool bEncryptIniFiles_DEPRECATED;
+
+	/**
+	* Encrypt the pak index
+	* NOTE: Replaced by the settings inside the cryptokeys system. Kept here for legacy migration purposes.
+	*/
+	UPROPERTY(config)
+	bool bEncryptPakIndex_DEPRECATED;
+
+	/**
+	 * Enable the early downloader pak file pakearly.txt
+	 * This has been superseded by the functionality in DefaultPakFileRules.ini
+	 */
+	UPROPERTY(config)
+	bool GenerateEarlyDownloaderPakFile_DEPRECATED;
+	
+	/**
+	 * Don't include content in any editor folders when cooking.  This can cause issues with missing content in cooked games if the content is being used. 
+	 */
+	UPROPERTY(config, EditAnywhere, Category = Packaging, AdvancedDisplay, meta = (DisplayName = "Exclude editor content when cooking"))
+	bool bSkipEditorContent;
+
+	/**
+	 * Don't include movies by default when staging/packaging
+	 * Specific movies can be specified below, and this can be in a platform ini
+	 */
+	UPROPERTY(config, EditAnywhere, Category = Packaging, AdvancedDisplay, meta = (DisplayName = "Exclude movie files when staging"))
+	bool bSkipMovies;
+
+	/**
+	 * If SkipMovies is true, these specific movies will still be added to the .pak file (if using a .pak file; otherwise they're copied as individual files)
+	 * This should be the name with no extension
+	 */
+	UPROPERTY(config, EditAnywhere, Category = Packaging, AdvancedDisplay, meta = (DisplayName = "Specific movies to Package"))
+	TArray<FString> UFSMovies;
+
+	/**
+	 * If SkipMovies is true, these specific movies will be copied when packaging your project, but are not supposed to be part of the .pak file
+	 * This should be the name with no extension
+	 */
+	UPROPERTY(config, EditAnywhere, Category = Packaging, AdvancedDisplay, meta = (DisplayName = "Specific movies to Copy"))
+	TArray<FString> NonUFSMovies;
+
+	/**
+	 * If set, only these specific pak files will be compressed. This should take the form of "*pakchunk0*"
+	 * This can be set in a platform-specific ini file
+	 */
+	UPROPERTY(config, EditAnywhere, Category = Packaging, AdvancedDisplay)
+	TArray<FString> CompressedChunkWildcard;
+
+	/** List of ini file keys to strip when packaging */
+	UPROPERTY(config, EditAnywhere, Category = Packaging)
+	TArray<FString> IniKeyBlacklist;
+
+	/** List of ini file sections to strip when packaging */
+	UPROPERTY(config, EditAnywhere, Category = Packaging)
+	TArray<FString> IniSectionBlacklist;
+
+	/**
+	 * List of specific files to include with GenerateEarlyDownloaderPakFile
+	 */
+	UPROPERTY(config)
+	TArray<FString> EarlyDownloaderPakFileFiles_DEPRECATED;
+
+
+	/**
+	 * List of maps to include when no other map list is specified on commandline
+	 */
+	UPROPERTY(config, EditAnywhere, Category = Packaging, AdvancedDisplay, meta = (DisplayName = "List of maps to include in a packaged build", RelativeToGameContentDir, LongPackageName))
+	TArray<FFilePath> MapsToCook;	
+
+	/**
+	 * Directories containing .uasset files that should always be cooked regardless of whether they're referenced by anything in your project
+	 * These paths are stored either as a full package path (e.g. /Game/Folder, /Engine/Folder, /PluginName/Folder) or as a relative package path from /Game
+	 */
+	UPROPERTY(config, EditAnywhere, Category=Packaging, AdvancedDisplay, meta=(DisplayName="Additional Asset Directories to Cook", LongPackageName))
+	TArray<FDirectoryPath> DirectoriesToAlwaysCook;
+
+	/**
+	 * Directories containing .uasset files that should never be cooked even if they are referenced by your project
+	 * These paths are stored either as a full package path (e.g. /Game/Folder, /Engine/Folder, /PluginName/Folder) or as a relative package path from /Game
+	 */
+	UPROPERTY(config, EditAnywhere, Category = Packaging, AdvancedDisplay, meta = (DisplayName = "Directories to never cook", LongPackageName))
+	TArray<FDirectoryPath> DirectoriesToNeverCook;
+
+	/**
+	 * Directories containing .uasset files that are for editor testing purposes and should not be included in
+	 * enumerations of all packages in a root directory, because they will cause errors on load
+	 * These paths are stored either as a full package path (e.g. /Game/Folder, /Engine/Folder, /PluginName/Folder) or as a relative package path from /Game
+	 */
+	UPROPERTY(config, EditAnywhere, Category = Packaging, AdvancedDisplay, meta = (DisplayName = "Test directories to not search", LongPackageName))
+	TArray<FDirectoryPath> TestDirectoriesToNotSearch;
+
+	/**
+	 * Directories containing files that should always be added to the .pak file (if using a .pak file; otherwise they're copied as individual files)
+	 * This is used to stage additional files that you manually load via the UFS (Unreal File System) file IO API
+	 * Note: These paths are relative to your project Content directory
+	 */
+	UPROPERTY(config, EditAnywhere, Category=Packaging, AdvancedDisplay, meta=(DisplayName="Additional Non-Asset Directories to Package", RelativeToGameContentDir))
+	TArray<FDirectoryPath> DirectoriesToAlwaysStageAsUFS;
+
+	/**
+	 * Directories containing files that should always be copied when packaging your project, but are not supposed to be part of the .pak file
+	 * This is used to stage additional files that you manually load without using the UFS (Unreal File System) file IO API, eg, third-party libraries that perform their own internal file IO
+	 * Note: These paths are relative to your project Content directory
+	 */
+	UPROPERTY(config, EditAnywhere, Category=Packaging, AdvancedDisplay, meta=(DisplayName="Additional Non-Asset Directories To Copy", RelativeToGameContentDir))
+	TArray<FDirectoryPath> DirectoriesToAlwaysStageAsNonUFS;	
+
+	/**
+	 * Directories containing files that should always be added to the .pak file for a dedicated server (if using a .pak file; otherwise they're copied as individual files)
+	 * This is used to stage additional files that you manually load via the UFS (Unreal File System) file IO API
+	 * Note: These paths are relative to your project Content directory
+	 */
+	UPROPERTY(config, EditAnywhere, Category=Packaging, AdvancedDisplay, meta=(DisplayName="Additional Non-Asset Directories to Package for dedicated server only", RelativeToGameContentDir))
+	TArray<FDirectoryPath> DirectoriesToAlwaysStageAsUFSServer;
+
+	/**
+	 * Directories containing files that should always be copied when packaging your project for a dedicated server, but are not supposed to be part of the .pak file
+	 * This is used to stage additional files that you manually load without using the UFS (Unreal File System) file IO API, eg, third-party libraries that perform their own internal file IO
+	 * Note: These paths are relative to your project Content directory
+	 */
+	UPROPERTY(config, EditAnywhere, Category=Packaging, AdvancedDisplay, meta=(DisplayName="Additional Non-Asset Directories To Copy for dedicated server only", RelativeToGameContentDir))
+	TArray<FDirectoryPath> DirectoriesToAlwaysStageAsNonUFSServer;	
+
+
+
+	/**
+	 * Get and set the per-platform build config and targetplatform settings for the Turnkey/Launch on menu
+	 */
+	EProjectPackagingBuildConfigurations GetBuildConfigurationForPlatform(FName PlatformName) const;
+	void SetBuildConfigurationForPlatform(FName PlatformName, EProjectPackagingBuildConfigurations Configuration);
+
+	FName GetTargetPlatformForPlatform(FName PlatformName) const;
+	void SetTargetPlatformForPlatform(FName PlatformName, FName TargetPlatformName);
+
+private:
+	/** Per platform build configuration */
+	UPROPERTY(config)
+	TMap<FName, EProjectPackagingBuildConfigurations> PerPlatformBuildConfig;
+
+	/** Per platform targetplatform cooking target */
+	UPROPERTY(config)
+	TMap<FName, FName> PerPlatformTargetPlatformName;
+
+		/** Helper array used to mirror Blueprint asset selections across edits */
+	TArray<FFilePath> CachedNativizeBlueprintAssets;
+	
+public:
+
+	// UObject Interface
+
+	virtual void PostInitProperties() override;
+#if WITH_EDITOR
+	virtual void PostEditChangeProperty( struct FPropertyChangedEvent& PropertyChangedEvent ) override;
+	virtual bool CanEditChange( const FProperty* InProperty ) const override;
+
+	/** Adds the given Blueprint asset to the exclusive nativization list. */
+	bool AddBlueprintAssetToNativizationList(const class UBlueprint* InBlueprint);
+
+	/** Removes the given Blueprint asset from the exclusive nativization list. */
+	bool RemoveBlueprintAssetFromNativizationList(const class UBlueprint* InBlueprint);
+
+	/** Determines if the specified Blueprint is already saved for exclusive nativization. */
+	bool IsBlueprintAssetInNativizationList(const class UBlueprint* InBlueprint) const { return FindBlueprintInNativizationList(InBlueprint) >= 0; }
+#endif
+
+	/** Gets a list of all valid packaging configurations for the current project */
+	static TArray<EProjectPackagingBuildConfigurations> GetValidPackageConfigurations();
+
+	/** Gets the current build target, checking that it's valid, and the default build target if it is not */
+	const FTargetInfo* GetBuildTargetInfo() const;
+
+private:
+	/** Returns the index of the specified Blueprint in the exclusive nativization list (otherwise INDEX_NONE) */
+	int32 FindBlueprintInNativizationList(const UBlueprint* InBlueprint) const;
+
+	/** Fix up cooking paths after they've been edited or laoded */
+	void FixCookingPaths();
+};
diff --git a/Engine/Source/Developer/TargetPlatform/Private/CookerSettings.cpp b/Engine/Source/Developer/TargetPlatform/Private/CookerSettings.cpp
new file mode 100644
--- /dev/null
+++ b/Engine/Source/Developer/TargetPlatform/Private/CookerSettings.cpp
@@ -0,0 +1,87 @@
+// Copyright Epic Games, Inc. All Rights Reserved.
+
+#include "CookerSettings.h"
+#include "UObject/UnrealType.h"
+#include "EngineLogs.h"
+
+UCookerSettings::UCookerSettings(const FObjectInitializer& ObjectInitializer)
+	: Super(ObjectInitializer)
+	, bEnableCookOnTheSide(false)
+	, bEnableBuildDDCInBackground(false)
+	, bIterativeCookingForLaunchOn(false)
+	, bCookOnTheFlyForLaunchOn(false)
+	, CookProgressDisplayMode(ECookProgressDisplayMode::RemainingPackages)
+	, bIgnoreIniSettingsOutOfDateForIteration(false)
+	, bIgnoreScriptPackagesOutOfDateForIteration(false)
+	, bCompileBlueprintsInDevelopmentMode(true)
+	, BlueprintComponentDataCookingMethod(EBlueprintComponentDataCookingMethod::EnabledBlueprintsOnly)
+	, bAllowCookedDataInEditorBuilds(false)
+	, bCookBlueprintComponentTemplateData(false)
+{
+	SectionName = TEXT("Cooker");
+	DefaultPVRTCQuality = 1;
+	DefaultASTCQualityBySize = 3;
+	DefaultASTCQualityBySpeed = 3;
+}
+
+void UCookerSettings::PostInitProperties()
+{
+	Super::PostInitProperties();
+	UObject::UpdateClassesExcludedFromDedicatedServer(ClassesExcludedOnDedicatedServer, ModulesExcludedOnDedicatedServer);
+	UObject::UpdateClassesExcludedFromDedicatedClient(ClassesExcludedOnDedicatedClient, ModulesExcludedOnDedicatedClient);
+
+	// In the 'false' case, we previously deferred to the 'EnabledBlueprintsOnly' method, which is the current default, so we don't need to handle it here.
+	if (bCookBlueprintComponentTemplateData)
+	{
+		BlueprintComponentDataCookingMethod = EBlueprintComponentDataCookingMethod::AllBlueprints;
+	}
+}
+
+#if WITH_EDITOR
+
+void UCookerSettings::PostEditChangeProperty(FPropertyChangedEvent& PropertyChangedEvent)
+{
+	static FName NAME_ClassesExcludedOnDedicatedServer = GET_MEMBER_NAME_CHECKED(UCookerSettings, ClassesExcludedOnDedicatedServer);
+	static FName NAME_ClassesExcludedOnDedicatedClient = GET_MEMBER_NAME_CHECKED(UCookerSettings, ClassesExcludedOnDedicatedClient);
+
+	static FName NAME_ModulesExcludedOnDedicatedServer(TEXT("ModulesExcludedOnDedicatedServer"));
+	static FName NAME_ModulesExcludedOnDedicatedClient(TEXT("ModulesExcludedOnDedicatedClient"));
+
+	static FName NAME_bCookBlueprintComponentTemplateData = GET_MEMBER_NAME_CHECKED(UCookerSettings, bCookBlueprintComponentTemplateData);
+	static FName NAME_BlueprintComponentDataCookingMethod = GET_MEMBER_NAME_CHECKED(UCookerSettings, BlueprintComponentDataCookingMethod);
+
+	if(PropertyChangedEvent.Property)
+	{
+		if(PropertyChangedEvent.Property->GetFName() == NAME_ClassesExcludedOnDedicatedServer
+			|| PropertyChangedEvent.Property->GetFName() == NAME_ModulesExcludedOnDedicatedServer)
+		{
+			UObject::UpdateClassesExcludedFromDedicatedServer(ClassesExcludedOnDedicatedServer, ModulesExcludedOnDedicatedServer);
+		}
+		else if(PropertyChangedEvent.Property->GetFName() == NAME_ClassesExcludedOnDedicatedClient
+			|| PropertyChangedEvent.Property->GetFName() == NAME_ModulesExcludedOnDedicatedClient)
+		{
+			UObject::UpdateClassesExcludedFromDedicatedClient(ClassesExcludedOnDedicatedClient, ModulesExcludedOnDedicatedClient);
+		}
+		else if (bCookBlueprintComponentTemplateData
+			&& PropertyChangedEvent.Property->GetFName() == NAME_BlueprintComponentDataCookingMethod)
+		{
+			UE_LOG(LogBlueprint, Warning, TEXT("\'%s\' has been deprecated in favor of \'%s\', please remove \'%s\' from %s"),
+				*NAME_bCookBlueprintComponentTemplateData.ToString(),
+				*NAME_BlueprintComponentDataCookingMethod.ToString(),
+				*NAME_bCookBlueprintComponentTemplateData.ToString(),
+				*GetDefaultConfigFilename());
+
+			if (const FBoolProperty* OldProperty = FindFieldChecked<FBoolProperty>(GetClass(), NAME_bCookBlueprintComponentTemplateData))
+			{
+				OldProperty->SetPropertyValue_InContainer(this, false);
+				UpdateSinglePropertyInConfigFile(OldProperty, GetDefaultConfigFilename());
+			}
+		}
+		else
+		{
+			ExportValuesToConsoleVariables(PropertyChangedEvent.Property);
+		}
+	}
+}
+
+#endif
\ No newline at end of file
diff --git a/Engine/Source/Developer/TargetPlatform/Private/ProjectPackagingSettings.cpp b/Engine/Source/Developer/TargetPlatform/Private/ProjectPackagingSettings.cpp
new file mode 100644
--- /dev/null
+++ b/Engine/Source/Developer/TargetPlatform/Private/ProjectPackagingSettings.cpp
@@ -0,0 +1,414 @@
+// Copyright Epic Games, Inc. All Rights Reserved.
+
+#include "Settings/ProjectPackagingSettings.h"
+#include "Misc/App.h"
+#include "Misc/ConfigCacheIni.h"
+#include "Interfaces/IProjectManager.h"
+#include "DesktopPlatformModule.h"
+#include "Common/TargetPlatformBase.h"
+
+#if WITH_EDITOR
+#include "Engine/Blueprint.h"
+#endif
+
+#define LOCTEXT_NAMESPACE "SettingsClasses"
+
+
+/* UProjectPackagingSettings interface
+ *****************************************************************************/
+
+const UProjectPackagingSettings::FConfigurationInfo UProjectPackagingSettings::ConfigurationInfo[(int)EProjectPackagingBuildConfigurations::PPBC_MAX] =
+{
+	/* PPBC_Debug */         { EBuildConfiguration::Debug, LOCTEXT("DebugConfiguration", "Debug"), LOCTEXT("DebugConfigurationTooltip", "Package the game in Debug configuration") },
+	/* PPBC_DebugGame */     { EBuildConfiguration::DebugGame, LOCTEXT("DebugGameConfiguration", "DebugGame"), LOCTEXT("DebugGameConfigurationTooltip", "Package the game in DebugGame configuration") },
+	/* PPBC_Development */   { EBuildConfiguration::Development, LOCTEXT("DevelopmentConfiguration", "Development"), LOCTEXT("DevelopmentConfigurationTooltip", "Package the game in Development configuration") },
+	/* PPBC_Test */          { EBuildConfiguration::Test, LOCTEXT("TestConfiguration", "Test"), LOCTEXT("TestConfigurationTooltip", "Package the game in Test configuration") },
+	/* PPBC_Shipping */      { EBuildConfiguration::Shipping, LOCTEXT("ShippingConfiguration", "Shipping"), LOCTEXT("ShippingConfigurationTooltip", "Package the game in Shipping configuration") },
+};
+
+UProjectPackagingSettings::UProjectPackagingSettings(const FObjectInitializer& ObjectInitializer)
+	: Super(ObjectInitializer)
+{
+}
+
+
+void UProjectPackagingSettings::PostInitProperties()
+{
+	// Build code projects by default
+	Build = EProjectPackagingBuild::IfProjectHasCode;
+
+	// Cache the current set of Blueprint assets selected for nativization.
+	CachedNativizeBlueprintAssets = NativizeBlueprintAssets;
+
+	FixCookingPaths();
+
+	Super::PostInitProperties();
+}
+
+void UProjectPackagingSettings::FixCookingPaths()
+{
+	// Fix AlwaysCook/NeverCook paths to use content root
+	for (FDirectoryPath& PathToFix : DirectoriesToAlwaysCook)
+	{
+		if (!PathToFix.Path.IsEmpty() && !PathToFix.Path.StartsWith(TEXT("/"), ESearchCase::CaseSensitive))
+		{
+			PathToFix.Path = FString::Printf(TEXT("/Game/%s"), *PathToFix.Path);
+		}
+	}
+
+	for (FDirectoryPath& PathToFix : DirectoriesToNeverCook)
+	{
+		if (!PathToFix.Path.IsEmpty() && !PathToFix.Path.StartsWith(TEXT("/"), ESearchCase::CaseSensitive))
+		{
+			PathToFix.Path = FString::Printf(TEXT("/Game/%s"), *PathToFix.Path);
+		}
+	}
+
+	for (FDirectoryPath& PathToFix : TestDirectoriesToNotSearch)
+	{
+		if (!PathToFix.Path.IsEmpty() && !PathToFix.Path.StartsWith(TEXT("/"), ESearchCase::CaseSensitive))
+		{
+			PathToFix.Path = FString::Printf(TEXT("/Game/%s"), *PathToFix.Path);
+		}
+	}
+}
+
+#if WITH_EDITOR
+void UProjectPackagingSettings::PostEditChangeProperty(FPropertyChangedEvent& PropertyChangedEvent)
+{
+	Super::PostEditChangeProperty(PropertyChangedEvent);
+
+	const FName Name = (PropertyChangedEvent.MemberProperty != nullptr)
+		? PropertyChangedEvent.MemberProperty->GetFName()
+		: NAME_None;
+
+	if (Name == FName(TEXT("DirectoriesToAlwaysCook")) || Name == FName(TEXT("DirectoriesToNeverCook")) || Name == FName(TEXT("TestDirectoriesToNotSearch")) || Name == NAME_None)
+	{
+		// We need to fix paths for no name updates to catch the reloadconfig call
+		FixCookingPaths();
+	}
+	else if (Name == FName((TEXT("StagingDirectory"))))
+	{
+		// fix up path
+		FString Path = StagingDirectory.Path;
+		FPaths::MakePathRelativeTo(Path, FPlatformProcess::BaseDir());
+		StagingDirectory.Path = Path;
+	}
+	else if (Name == FName(TEXT("ForDistribution")))
+	{
+		if (ForDistribution && BuildConfiguration != EProjectPackagingBuildConfigurations::PPBC_Shipping)
+		{
+			BuildConfiguration = EProjectPackagingBuildConfigurations::PPBC_Shipping;
+			// force serialization for "Build COnfiguration"
+			UpdateSinglePropertyInConfigFile(GetClass()->FindPropertyByName(GET_MEMBER_NAME_CHECKED(UProjectPackagingSettings, BuildConfiguration)), GetDefaultConfigFilename());
+		}
+	}
+	else if (Name == FName(TEXT("bGenerateChunks")))
+	{
+		if (bGenerateChunks)
+		{
+			UsePakFile = true;
+		}
+	}
+	else if (Name == FName(TEXT("UsePakFile")))
+	{
+		if (!UsePakFile)
+		{
+			bGenerateChunks = false;
+			bBuildHttpChunkInstallData = false;
+		}
+	}
+	else if (Name == FName(TEXT("bBuildHTTPChunkInstallData")))
+	{
+		if (bBuildHttpChunkInstallData)
+		{
+			UsePakFile = true;
+			bGenerateChunks = true;
+			//Ensure data is something valid
+			if (HttpChunkInstallDataDirectory.Path.IsEmpty())
+			{
+				auto CloudInstallDir = FPaths::ConvertRelativePathToFull(FPaths::GetPath(FPaths::GetProjectFilePath())) / TEXT("ChunkInstall");
+				HttpChunkInstallDataDirectory.Path = CloudInstallDir;
+			}
+			if (HttpChunkInstallDataVersion.IsEmpty())
+			{
+				HttpChunkInstallDataVersion = TEXT("release1");
+			}
+		}
+	}
+	else if (Name == FName((TEXT("ApplocalPrerequisitesDirectory"))))
+	{
+		// If a variable is already in use, assume the user knows what they are doing and don't modify the path
+		if (!ApplocalPrerequisitesDirectory.Path.Contains("$("))
+		{
+			// Try making the path local to either project or engine directories.
+			FString EngineRootedPath = ApplocalPrerequisitesDirectory.Path;
+			FString EnginePath = FPaths::ConvertRelativePathToFull(FPaths::GetPath(FPaths::EngineDir())) + "/";
+			FPaths::MakePathRelativeTo(EngineRootedPath, *EnginePath);
+			if (FPaths::IsRelative(EngineRootedPath))
+			{
+				ApplocalPrerequisitesDirectory.Path = "$(EngineDir)/" + EngineRootedPath;
+				return;
+			}
+
+			FString ProjectRootedPath = ApplocalPrerequisitesDirectory.Path;
+			FString ProjectPath = FPaths::ConvertRelativePathToFull(FPaths::GetPath(FPaths::GetProjectFilePath())) + "/";
+			FPaths::MakePathRelativeTo(ProjectRootedPath, *ProjectPath);
+			if (FPaths::IsRelative(EngineRootedPath))
+			{
+				ApplocalPrerequisitesDirectory.Path = "$(ProjectDir)/" + ProjectRootedPath;
+				return;
+			}
+		}
+	}
+	else if (Name == FName((TEXT("NativizeBlueprintAssets"))))
+	{
+		int32 AssetIndex;
+		auto OnSelectBlueprintForExclusiveNativizationLambda = [](const FString& PackageName, bool bSelect)
+		{
+			if (!PackageName.IsEmpty())
+			{
+				// This should only apply to loaded packages. Any unloaded packages defer setting the transient flag to when they're loaded.
+				if (UPackage* Package = FindPackage(nullptr, *PackageName))
+				{
+					// Find the Blueprint asset within the package.
+					if (UBlueprint* Blueprint = FindObject<UBlueprint>(Package, *FPaths::GetBaseFilename(PackageName)))
+					{
+						// We're toggling the transient flag on or off.
+						if ((Blueprint->NativizationFlag == EBlueprintNativizationFlag::ExplicitlyEnabled) != bSelect)
+						{
+							Blueprint->NativizationFlag = bSelect ? EBlueprintNativizationFlag::ExplicitlyEnabled : EBlueprintNativizationFlag::Disabled;
+						}
+					}
+				}
+			}
+		};
+
+		if (NativizeBlueprintAssets.Num() > 0)
+		{
+			for (AssetIndex = 0; AssetIndex < NativizeBlueprintAssets.Num(); ++AssetIndex)
+			{
+				const FString& PackageName = NativizeBlueprintAssets[AssetIndex].FilePath;
+				if (AssetIndex >= CachedNativizeBlueprintAssets.Num())
+				{
+					// A new entry was added; toggle the exclusive flag on the corresponding Blueprint asset (if loaded).
+					OnSelectBlueprintForExclusiveNativizationLambda(PackageName, true);
+
+					// Add an entry to the end of the cached list.
+					CachedNativizeBlueprintAssets.Add(NativizeBlueprintAssets[AssetIndex]);
+				}
+				else if (!PackageName.Equals(CachedNativizeBlueprintAssets[AssetIndex].FilePath))
+				{
+					if (NativizeBlueprintAssets.Num() < CachedNativizeBlueprintAssets.Num())
+					{
+						// An entry was removed; toggle the exclusive flag on the corresponding Blueprint asset (if loaded).
+						OnSelectBlueprintForExclusiveNativizationLambda(CachedNativizeBlueprintAssets[AssetIndex].FilePath, false);
+
+						// Remove this entry from the cached list.
+						CachedNativizeBlueprintAssets.RemoveAt(AssetIndex);
+					}
+					else if (NativizeBlueprintAssets.Num() > CachedNativizeBlueprintAssets.Num())
+					{
+						// A new entry was inserted; toggle the exclusive flag on the corresponding Blueprint asset (if loaded).
+						OnSelectBlueprintForExclusiveNativizationLambda(PackageName, true);
+
+						// Insert the new entry into the cached list.
+						CachedNativizeBlueprintAssets.Insert(NativizeBlueprintAssets[AssetIndex], AssetIndex);
+					}
+					else
+					{
+						// An entry was changed; toggle the exclusive flag on the corresponding Blueprint assets (if loaded).
+						OnSelectBlueprintForExclusiveNativizationLambda(CachedNativizeBlueprintAssets[AssetIndex].FilePath, false);
+						OnSelectBlueprintForExclusiveNativizationLambda(PackageName, true);
+
+						// Update the cached entry.
+						CachedNativizeBlueprintAssets[AssetIndex].FilePath = PackageName;
+					}
+				}
+			}
+
+			if (CachedNativizeBlueprintAssets.Num() > NativizeBlueprintAssets.Num())
+			{
+				// Removed entries at the end of the list; toggle the exclusive flag on the corresponding Blueprint asset(s) (if loaded).
+				for (AssetIndex = NativizeBlueprintAssets.Num(); AssetIndex < CachedNativizeBlueprintAssets.Num(); ++AssetIndex)
+				{
+					OnSelectBlueprintForExclusiveNativizationLambda(CachedNativizeBlueprintAssets[AssetIndex].FilePath, false);
+				}
+
+				// Remove entries from the end of the cached list.
+				CachedNativizeBlueprintAssets.RemoveAt(NativizeBlueprintAssets.Num(), CachedNativizeBlueprintAssets.Num() - NativizeBlueprintAssets.Num());
+			}
+		}
+		else if (CachedNativizeBlueprintAssets.Num() > 0)
+		{
+			// Removed all entries; toggle the exclusive flag on the corresponding Blueprint asset(s) (if loaded).
+			for (AssetIndex = 0; AssetIndex < CachedNativizeBlueprintAssets.Num(); ++AssetIndex)
+			{
+				OnSelectBlueprintForExclusiveNativizationLambda(CachedNativizeBlueprintAssets[AssetIndex].FilePath, false);
+			}
+
+			// Clear the cached list.
+			CachedNativizeBlueprintAssets.Empty();
+		}
+	}
+}
+
+bool UProjectPackagingSettings::CanEditChange(const FProperty* InProperty) const
+{
+	if (InProperty->GetFName() == FName(TEXT("NativizeBlueprintAssets")))
+	{
+		return BlueprintNativizationMethod == EProjectPackagingBlueprintNativizationMethod::Exclusive;
+	}
+
+	return Super::CanEditChange(InProperty);
+}
+
+
+bool UProjectPackagingSettings::AddBlueprintAssetToNativizationList(const class UBlueprint* InBlueprint)
+{
+	if (InBlueprint)
+	{
+		const FString PackageName = InBlueprint->GetOutermost()->GetName();
+
+		// Make sure it's not already in the exclusive list. This can happen if the user previously added this asset in the Project Settings editor.
+		const bool bFound = IsBlueprintAssetInNativizationList(InBlueprint);
+		if (!bFound)
+		{
+			// Add this Blueprint asset to the exclusive list.
+			FFilePath FileInfo;
+			FileInfo.FilePath = PackageName;
+			NativizeBlueprintAssets.Add(FileInfo);
+
+			// Also add it to the mirrored list for tracking edits.
+			CachedNativizeBlueprintAssets.Add(FileInfo);
+
+			return true;
+		}
+	}
+
+	return false;
+}
+
+bool UProjectPackagingSettings::RemoveBlueprintAssetFromNativizationList(const class UBlueprint* InBlueprint)
+{
+	if (InBlueprint)
+	{
+		const FString PackageName = InBlueprint->GetOutermost()->GetName();
+
+		int32 AssetIndex = FindBlueprintInNativizationList(InBlueprint);
+		if (AssetIndex >= 0)
+		{
+			// Note: Intentionally not using RemoveAtSwap() here, so that the order is preserved.
+			NativizeBlueprintAssets.RemoveAt(AssetIndex);
+
+			// Also remove it from the mirrored list (for tracking edits).
+			CachedNativizeBlueprintAssets.RemoveAt(AssetIndex);
+
+			return true;
+		}
+	}
+
+	return false;
+}
+
+int32 UProjectPackagingSettings::FindBlueprintInNativizationList(const UBlueprint* InBlueprint) const
+{
+	int32 ListIndex = INDEX_NONE;
+	if (InBlueprint)
+	{
+		const FString PackageName = InBlueprint->GetOutermost()->GetName();
+
+		for (int32 AssetIndex = 0; AssetIndex < NativizeBlueprintAssets.Num(); ++AssetIndex)
+		{
+			if (NativizeBlueprintAssets[AssetIndex].FilePath.Equals(PackageName, ESearchCase::IgnoreCase))
+			{
+				ListIndex = AssetIndex;
+				break;
+			}
+		}
+	}
+	return ListIndex;
+}
+
+
+#endif
+
+TArray<EProjectPackagingBuildConfigurations> UProjectPackagingSettings::GetValidPackageConfigurations()
+{
+	// Check if the project has code
+	FProjectStatus ProjectStatus;
+	bool bHasCode = IProjectManager::Get().QueryStatusForCurrentProject(ProjectStatus) && ProjectStatus.bCodeBasedProject;
+
+	// If if does, find all the targets
+	const TArray<FTargetInfo>* Targets = nullptr;
+	if (bHasCode)
+	{
+		Targets = &(FDesktopPlatformModule::Get()->GetTargetsForCurrentProject());
+	}
+
+	// Set up all the configurations
+	TArray<EProjectPackagingBuildConfigurations> Configurations;
+	for (int32 Idx = 0; Idx < (int)EProjectPackagingBuildConfigurations::PPBC_MAX; Idx++)
+	{
+		EProjectPackagingBuildConfigurations PackagingConfiguration = (EProjectPackagingBuildConfigurations)Idx;
+
+		// Check the target type is valid
+		const UProjectPackagingSettings::FConfigurationInfo& Info = UProjectPackagingSettings::ConfigurationInfo[Idx];
+		if (!bHasCode && Info.Configuration == EBuildConfiguration::DebugGame)
+		{
+			continue;
+		}
+
+		Configurations.Add(PackagingConfiguration);
+	}
+	return Configurations;
+}
+
+const FTargetInfo* UProjectPackagingSettings::GetBuildTargetInfo() const
+{
+	const FTargetInfo* DefaultGameTarget = nullptr;
+	const FTargetInfo* DefaultClientTarget = nullptr;
+	for (const FTargetInfo& Target : FDesktopPlatformModule::Get()->GetTargetsForCurrentProject())
+	{
+		if (Target.Name == BuildTarget)
+		{
+			return &Target;
+		}
+		else if (Target.Type == EBuildTargetType::Game && (DefaultGameTarget == nullptr || Target.Name < DefaultGameTarget->Name))
+		{
+			DefaultGameTarget = &Target;
+		}
+		else if (Target.Type == EBuildTargetType::Client && (DefaultClientTarget == nullptr || Target.Name < DefaultClientTarget->Name))
+		{
+			DefaultClientTarget = &Target;
+		}
+	}
+	return (DefaultGameTarget != nullptr) ? DefaultGameTarget : DefaultClientTarget;
+}
+
+EProjectPackagingBuildConfigurations UProjectPackagingSettings::GetBuildConfigurationForPlatform(FName PlatformName) const
+{
+	const EProjectPackagingBuildConfigurations* Value = PerPlatformBuildConfig.Find(PlatformName);
+
+	return Value == nullptr ? EProjectPackagingBuildConfigurations::PPBC_Development : *Value;
+}
+
+void UProjectPackagingSettings::SetBuildConfigurationForPlatform(FName PlatformName, EProjectPackagingBuildConfigurations Configuration)
+{
+	PerPlatformBuildConfig.Add(PlatformName, Configuration);
+}
+
+FName UProjectPackagingSettings::GetTargetPlatformForPlatform(FName PlatformName) const
+{
+	const FName* Value = PerPlatformTargetPlatformName.Find(PlatformName);
+
+	// the platform name is also the name of the vanilla info
+	return Value == nullptr ? PlatformName : *Value;
+}
+
+void UProjectPackagingSettings::SetTargetPlatformForPlatform(FName PlatformName, FName TargetPlatformName)
+{
+	PerPlatformTargetPlatformName.Add(PlatformName, TargetPlatformName);
+}
+
+#undef LOCTEXT_NAMESPACE
\ No newline at end of file
diff --git a/Engine/Source/Developer/TargetPlatform/TargetPlatform.Build.cs b/Engine/Source/Developer/TargetPlatform/TargetPlatform.Build.cs
--- a/Engine/Source/Developer/TargetPlatform/TargetPlatform.Build.cs
+++ b/Engine/Source/Developer/TargetPlatform/TargetPlatform.Build.cs
@@ -8,16 +8,25 @@ public class TargetPlatform : ModuleRules
 	public TargetPlatform(ReadOnlyTargetRules Target) : base(Target)
 	{
 		PrivateDependencyModuleNames.Add("Core");
+		PrivateDependencyModuleNames.Add("CoreUObject");
 		PrivateDependencyModuleNames.Add("SlateCore");
 		PrivateDependencyModuleNames.Add("Slate");
 		PrivateDependencyModuleNames.Add("EditorStyle");
 		PrivateDependencyModuleNames.Add("Projects");
+		PublicDependencyModuleNames.Add("DeveloperSettings");
 		PublicDependencyModuleNames.Add("AudioPlatformConfiguration");
 		PublicDependencyModuleNames.Add("DesktopPlatform");
 		PublicDependencyModuleNames.Add("LauncherPlatform");
 
 		PrivateIncludePathModuleNames.Add("Engine");
 		PrivateIncludePathModuleNames.Add("PhysicsCore");
+
+		if (Target.bCompileAgainstEngine)
+		{
+			PrivateDependencyModuleNames.Add("Engine");
+		}
+
+		DynamicallyLoadedModuleNames.Add("TurnkeySupport");
 		PrivateIncludePathModuleNames.Add("TurnkeySupport");
 
 		// no need for all these modules if the program doesn't want developer tools at all (like UnrealFileServer)
diff --git a/Engine/Source/Editor/ProjectSettingsViewer/ProjectSettingsViewer.Build.cs b/Engine/Source/Editor/ProjectSettingsViewer/ProjectSettingsViewer.Build.cs
--- a/Engine/Source/Editor/ProjectSettingsViewer/ProjectSettingsViewer.Build.cs
+++ b/Engine/Source/Editor/ProjectSettingsViewer/ProjectSettingsViewer.Build.cs
@@ -29,6 +29,7 @@ namespace UnrealBuildTool.Rules
 					"AIModule",
 					"ProjectTargetPlatformEditor",
 					"EditorStyle",
+					"TargetPlatform",
 				}
 			);
 
diff --git a/Engine/Source/Editor/UnrealEd/Classes/Settings/ProjectPackagingSettings.h b/Engine/Source/Editor/UnrealEd/Classes/Settings/ProjectPackagingSettings.h
deleted file mode 100644
diff --git a/Engine/Source/Editor/UnrealEd/Private/CookerSettings.cpp b/Engine/Source/Editor/UnrealEd/Private/CookerSettings.cpp
deleted file mode 100644
diff --git a/Engine/Source/Editor/UnrealEd/Private/Settings/SettingsClasses.cpp b/Engine/Source/Editor/UnrealEd/Private/Settings/SettingsClasses.cpp
--- a/Engine/Source/Editor/UnrealEd/Private/Settings/SettingsClasses.cpp
+++ b/Engine/Source/Editor/UnrealEd/Private/Settings/SettingsClasses.cpp
@@ -20,7 +20,6 @@
 #include "Settings/EditorLoadingSavingSettings.h"
 #include "Settings/EditorMiscSettings.h"
 #include "Settings/LevelEditorMiscSettings.h"
-#include "Settings/ProjectPackagingSettings.h"
 #include "EngineGlobals.h"
 #include "Components/ArrowComponent.h"
 #include "Components/BillboardComponent.h"
@@ -1004,398 +1003,6 @@ void ULevelEditorViewportSettings::PostEditChangeProperty( struct FPropertyChang
 	SettingChangedEvent.Broadcast(Name);
 }
 
-/* UProjectPackagingSettings interface
- *****************************************************************************/
-
-const UProjectPackagingSettings::FConfigurationInfo UProjectPackagingSettings::ConfigurationInfo[(int)EProjectPackagingBuildConfigurations::PPBC_MAX] =
-{
-	/* PPBC_Debug */         { EBuildConfiguration::Debug, LOCTEXT("DebugConfiguration", "Debug"), LOCTEXT("DebugConfigurationTooltip", "Package the game in Debug configuration") },
-	/* PPBC_DebugGame */     { EBuildConfiguration::DebugGame, LOCTEXT("DebugGameConfiguration", "DebugGame"), LOCTEXT("DebugGameConfigurationTooltip", "Package the game in DebugGame configuration") },
-	/* PPBC_Development */   { EBuildConfiguration::Development, LOCTEXT("DevelopmentConfiguration", "Development"), LOCTEXT("DevelopmentConfigurationTooltip", "Package the game in Development configuration") },
-	/* PPBC_Test */          { EBuildConfiguration::Test, LOCTEXT("TestConfiguration", "Test"), LOCTEXT("TestConfigurationTooltip", "Package the game in Test configuration") },
-	/* PPBC_Shipping */      { EBuildConfiguration::Shipping, LOCTEXT("ShippingConfiguration", "Shipping"), LOCTEXT("ShippingConfigurationTooltip", "Package the game in Shipping configuration") },
-};
-
-UProjectPackagingSettings::UProjectPackagingSettings( const FObjectInitializer& ObjectInitializer )
-	: Super(ObjectInitializer)
-{
-}
-
-
-void UProjectPackagingSettings::PostInitProperties()
-{
-	// Build code projects by default
-	Build = EProjectPackagingBuild::IfProjectHasCode;
-
-	// Cache the current set of Blueprint assets selected for nativization.
-	CachedNativizeBlueprintAssets = NativizeBlueprintAssets;
-
-	FixCookingPaths();
-
-	Super::PostInitProperties();
-}
-
-void UProjectPackagingSettings::FixCookingPaths()
-{
-	// Fix AlwaysCook/NeverCook paths to use content root
-	for (FDirectoryPath& PathToFix : DirectoriesToAlwaysCook)
-	{
-		if (!PathToFix.Path.IsEmpty() && !PathToFix.Path.StartsWith(TEXT("/"), ESearchCase::CaseSensitive))
-		{
-			PathToFix.Path = FString::Printf(TEXT("/Game/%s"), *PathToFix.Path);
-		}
-	}
-
-	for (FDirectoryPath& PathToFix : DirectoriesToNeverCook)
-	{
-		if (!PathToFix.Path.IsEmpty() && !PathToFix.Path.StartsWith(TEXT("/"), ESearchCase::CaseSensitive))
-		{
-			PathToFix.Path = FString::Printf(TEXT("/Game/%s"), *PathToFix.Path);
-		}
-	}
-
-	for (FDirectoryPath& PathToFix : TestDirectoriesToNotSearch)
-	{
-		if (!PathToFix.Path.IsEmpty() && !PathToFix.Path.StartsWith(TEXT("/"), ESearchCase::CaseSensitive))
-		{
-			PathToFix.Path = FString::Printf(TEXT("/Game/%s"), *PathToFix.Path);
-		}
-	}
-}
-
-void UProjectPackagingSettings::PostEditChangeProperty( FPropertyChangedEvent& PropertyChangedEvent )
-{
-	Super::PostEditChangeProperty(PropertyChangedEvent);
-
-	const FName Name = (PropertyChangedEvent.MemberProperty != nullptr)
-		? PropertyChangedEvent.MemberProperty->GetFName()
-		: NAME_None;
-
-	if (Name == FName(TEXT("DirectoriesToAlwaysCook")) || Name == FName(TEXT("DirectoriesToNeverCook")) || Name == FName(TEXT("TestDirectoriesToNotSearch")) || Name == NAME_None)
-	{
-		// We need to fix paths for no name updates to catch the reloadconfig call
-		FixCookingPaths();
-	}
-	else if (Name == FName((TEXT("StagingDirectory"))))
-	{
-		// fix up path
-		FString Path = StagingDirectory.Path;
-		FPaths::MakePathRelativeTo(Path, FPlatformProcess::BaseDir());
-		StagingDirectory.Path = Path;
-	}
-	else if (Name == FName(TEXT("ForDistribution")))
-	{
-		if (ForDistribution && BuildConfiguration != EProjectPackagingBuildConfigurations::PPBC_Shipping)
-		{
-			BuildConfiguration = EProjectPackagingBuildConfigurations::PPBC_Shipping;
-			// force serialization for "Build COnfiguration"
-			UpdateSinglePropertyInConfigFile(GetClass()->FindPropertyByName(GET_MEMBER_NAME_CHECKED(UProjectPackagingSettings, BuildConfiguration)), GetDefaultConfigFilename());
-		}
-	}
-	else if (Name == FName(TEXT("bGenerateChunks")))
-	{
-		if (bGenerateChunks)
-		{
-			UsePakFile = true;
-		}
-	}
-	else if (Name == FName(TEXT("UsePakFile")))
-	{
-		if (!UsePakFile)
-		{
-			bGenerateChunks = false;
-			bBuildHttpChunkInstallData = false;
-		}
-	}
-	else if (Name == FName(TEXT("bBuildHTTPChunkInstallData")))
-	{
-		if (bBuildHttpChunkInstallData)
-		{
-			UsePakFile = true;
-			bGenerateChunks = true;
-			//Ensure data is something valid
-			if (HttpChunkInstallDataDirectory.Path.IsEmpty())
-			{
-				auto CloudInstallDir = FPaths::ConvertRelativePathToFull(FPaths::GetPath(FPaths::GetProjectFilePath())) / TEXT("ChunkInstall");
-				HttpChunkInstallDataDirectory.Path = CloudInstallDir;
-			}
-			if (HttpChunkInstallDataVersion.IsEmpty())
-			{
-				HttpChunkInstallDataVersion = TEXT("release1");
-			}
-		}
-	}
-	else if (Name == FName((TEXT("ApplocalPrerequisitesDirectory"))))
-	{
-		// If a variable is already in use, assume the user knows what they are doing and don't modify the path
-		if(!ApplocalPrerequisitesDirectory.Path.Contains("$("))
-		{
-			// Try making the path local to either project or engine directories.
-			FString EngineRootedPath = ApplocalPrerequisitesDirectory.Path;
-			FString EnginePath = FPaths::ConvertRelativePathToFull(FPaths::GetPath(FPaths::EngineDir())) + "/";
-			FPaths::MakePathRelativeTo(EngineRootedPath, *EnginePath);
-			if (FPaths::IsRelative(EngineRootedPath))
-			{
-				ApplocalPrerequisitesDirectory.Path = "$(EngineDir)/" + EngineRootedPath;
-				return;
-			}
-
-			FString ProjectRootedPath = ApplocalPrerequisitesDirectory.Path;
-			FString ProjectPath = FPaths::ConvertRelativePathToFull(FPaths::GetPath(FPaths::GetProjectFilePath())) + "/";
-			FPaths::MakePathRelativeTo(ProjectRootedPath, *ProjectPath);
-			if (FPaths::IsRelative(EngineRootedPath))
-			{
-				ApplocalPrerequisitesDirectory.Path = "$(ProjectDir)/" + ProjectRootedPath;
-				return;
-			}
-		}
-	}
-	else if (Name == FName((TEXT("NativizeBlueprintAssets"))))
-	{
-		int32 AssetIndex;
-		auto OnSelectBlueprintForExclusiveNativizationLambda = [](const FString& PackageName, bool bSelect)
-		{
-			if (!PackageName.IsEmpty())
-			{
-				// This should only apply to loaded packages. Any unloaded packages defer setting the transient flag to when they're loaded.
-				if (UPackage* Package = FindPackage(nullptr, *PackageName))
-				{
-					// Find the Blueprint asset within the package.
-					if (UBlueprint* Blueprint = FindObject<UBlueprint>(Package, *FPaths::GetBaseFilename(PackageName)))
-					{
-						// We're toggling the transient flag on or off.
-						if ((Blueprint->NativizationFlag == EBlueprintNativizationFlag::ExplicitlyEnabled) != bSelect)
-						{
-							Blueprint->NativizationFlag = bSelect ? EBlueprintNativizationFlag::ExplicitlyEnabled : EBlueprintNativizationFlag::Disabled;
-						}
-					}
-				}
-			}
-		};
-
-		if (NativizeBlueprintAssets.Num() > 0)
-		{
-			for (AssetIndex = 0; AssetIndex < NativizeBlueprintAssets.Num(); ++AssetIndex)
-			{
-				const FString& PackageName = NativizeBlueprintAssets[AssetIndex].FilePath;
-				if (AssetIndex >= CachedNativizeBlueprintAssets.Num())
-				{
-					// A new entry was added; toggle the exclusive flag on the corresponding Blueprint asset (if loaded).
-					OnSelectBlueprintForExclusiveNativizationLambda(PackageName, true);
-
-					// Add an entry to the end of the cached list.
-					CachedNativizeBlueprintAssets.Add(NativizeBlueprintAssets[AssetIndex]);
-				}
-				else if (!PackageName.Equals(CachedNativizeBlueprintAssets[AssetIndex].FilePath))
-				{
-					if (NativizeBlueprintAssets.Num() < CachedNativizeBlueprintAssets.Num())
-					{
-						// An entry was removed; toggle the exclusive flag on the corresponding Blueprint asset (if loaded).
-						OnSelectBlueprintForExclusiveNativizationLambda(CachedNativizeBlueprintAssets[AssetIndex].FilePath, false);
-
-						// Remove this entry from the cached list.
-						CachedNativizeBlueprintAssets.RemoveAt(AssetIndex);
-					}
-					else if(NativizeBlueprintAssets.Num() > CachedNativizeBlueprintAssets.Num())
-					{
-						// A new entry was inserted; toggle the exclusive flag on the corresponding Blueprint asset (if loaded).
-						OnSelectBlueprintForExclusiveNativizationLambda(PackageName, true);
-
-						// Insert the new entry into the cached list.
-						CachedNativizeBlueprintAssets.Insert(NativizeBlueprintAssets[AssetIndex], AssetIndex);
-					}
-					else
-					{
-						// An entry was changed; toggle the exclusive flag on the corresponding Blueprint assets (if loaded).
-						OnSelectBlueprintForExclusiveNativizationLambda(CachedNativizeBlueprintAssets[AssetIndex].FilePath, false);
-						OnSelectBlueprintForExclusiveNativizationLambda(PackageName, true);
-
-						// Update the cached entry.
-						CachedNativizeBlueprintAssets[AssetIndex].FilePath = PackageName;
-					}
-				}
-			}
-
-			if (CachedNativizeBlueprintAssets.Num() > NativizeBlueprintAssets.Num())
-			{
-				// Removed entries at the end of the list; toggle the exclusive flag on the corresponding Blueprint asset(s) (if loaded).
-				for (AssetIndex = NativizeBlueprintAssets.Num(); AssetIndex < CachedNativizeBlueprintAssets.Num(); ++AssetIndex)
-				{
-					OnSelectBlueprintForExclusiveNativizationLambda(CachedNativizeBlueprintAssets[AssetIndex].FilePath, false);
-				}
-
-				// Remove entries from the end of the cached list.
-				CachedNativizeBlueprintAssets.RemoveAt(NativizeBlueprintAssets.Num(), CachedNativizeBlueprintAssets.Num() - NativizeBlueprintAssets.Num());
-			}
-		}
-		else if(CachedNativizeBlueprintAssets.Num() > 0)
-		{
-			// Removed all entries; toggle the exclusive flag on the corresponding Blueprint asset(s) (if loaded).
-			for (AssetIndex = 0; AssetIndex < CachedNativizeBlueprintAssets.Num(); ++AssetIndex)
-			{
-				OnSelectBlueprintForExclusiveNativizationLambda(CachedNativizeBlueprintAssets[AssetIndex].FilePath, false);
-			}
-
-			// Clear the cached list.
-			CachedNativizeBlueprintAssets.Empty();
-		}
-	}
-}
-
-bool UProjectPackagingSettings::CanEditChange( const FProperty* InProperty ) const
-{
-	if (InProperty->GetFName() == FName(TEXT("NativizeBlueprintAssets")))
-	{
-		return BlueprintNativizationMethod == EProjectPackagingBlueprintNativizationMethod::Exclusive;
-	}
-
-	return Super::CanEditChange(InProperty);
-}
-
-bool UProjectPackagingSettings::AddBlueprintAssetToNativizationList(const class UBlueprint* InBlueprint)
-{
-	if (InBlueprint)
-	{
-		const FString PackageName = InBlueprint->GetOutermost()->GetName();
-
-		// Make sure it's not already in the exclusive list. This can happen if the user previously added this asset in the Project Settings editor.
-		const bool bFound = IsBlueprintAssetInNativizationList(InBlueprint);
-		if (!bFound)
-		{
-			// Add this Blueprint asset to the exclusive list.
-			FFilePath FileInfo;
-			FileInfo.FilePath = PackageName;
-			NativizeBlueprintAssets.Add(FileInfo);
-
-			// Also add it to the mirrored list for tracking edits.
-			CachedNativizeBlueprintAssets.Add(FileInfo);
-
-			return true;
-		}
-	}
-
-	return false;
-}
-
-bool UProjectPackagingSettings::RemoveBlueprintAssetFromNativizationList(const class UBlueprint* InBlueprint)
-{
-	if (InBlueprint)
-	{
-		const FString PackageName = InBlueprint->GetOutermost()->GetName();
-
-		int32 AssetIndex = FindBlueprintInNativizationList(InBlueprint);
-		if (AssetIndex >= 0)
-		{
-			// Note: Intentionally not using RemoveAtSwap() here, so that the order is preserved.
-			NativizeBlueprintAssets.RemoveAt(AssetIndex);
-
-			// Also remove it from the mirrored list (for tracking edits).
-			CachedNativizeBlueprintAssets.RemoveAt(AssetIndex);
-
-			return true;
-		}
-	}
-
-	return false;
-}
-
-TArray<EProjectPackagingBuildConfigurations> UProjectPackagingSettings::GetValidPackageConfigurations()
-{
-	// Check if the project has code
-	FProjectStatus ProjectStatus;
-	bool bHasCode = IProjectManager::Get().QueryStatusForCurrentProject(ProjectStatus) && ProjectStatus.bCodeBasedProject;
-
-	// If if does, find all the targets
-	const TArray<FTargetInfo>* Targets = nullptr;
-	if (bHasCode)
-	{
-		Targets = &(FDesktopPlatformModule::Get()->GetTargetsForCurrentProject());
-	}
-
-	// Set up all the configurations
-	TArray<EProjectPackagingBuildConfigurations> Configurations;
-	for (int32 Idx = 0; Idx < (int)EProjectPackagingBuildConfigurations::PPBC_MAX; Idx++)
-	{
-		EProjectPackagingBuildConfigurations PackagingConfiguration = (EProjectPackagingBuildConfigurations)Idx;
-
-		// Check the target type is valid
-		const UProjectPackagingSettings::FConfigurationInfo& Info = UProjectPackagingSettings::ConfigurationInfo[Idx];
-		if(!bHasCode && Info.Configuration == EBuildConfiguration::DebugGame)
-		{
-			continue;
-		}
-
-		Configurations.Add(PackagingConfiguration);
-	}
-	return Configurations;
-}
-
-const FTargetInfo* UProjectPackagingSettings::GetBuildTargetInfo() const
-{
-	const FTargetInfo* DefaultGameTarget = nullptr;
-	const FTargetInfo* DefaultClientTarget = nullptr;
-	for (const FTargetInfo& Target : FDesktopPlatformModule::Get()->GetTargetsForCurrentProject())
-	{
-		if (Target.Name == BuildTarget)
-		{
-			return &Target;
-		}
-		else if (Target.Type == EBuildTargetType::Game && (DefaultGameTarget == nullptr || Target.Name < DefaultGameTarget->Name))
-		{
-			DefaultGameTarget = &Target;
-		}
-		else if (Target.Type == EBuildTargetType::Client && (DefaultClientTarget == nullptr || Target.Name < DefaultClientTarget->Name))
-		{
-			DefaultClientTarget = &Target;
-		}
-	}
-	return (DefaultGameTarget != nullptr)? DefaultGameTarget : DefaultClientTarget;
-}
-
-int32 UProjectPackagingSettings::FindBlueprintInNativizationList(const UBlueprint* InBlueprint) const
-{
-	int32 ListIndex = INDEX_NONE;
-	if (InBlueprint)
-	{
-		const FString PackageName = InBlueprint->GetOutermost()->GetName();
-
-		for (int32 AssetIndex = 0; AssetIndex < NativizeBlueprintAssets.Num(); ++AssetIndex)
-		{
-			if (NativizeBlueprintAssets[AssetIndex].FilePath.Equals(PackageName, ESearchCase::IgnoreCase))
-			{
-				ListIndex = AssetIndex;
-				break;
-			}
-		}
-	}
-	return ListIndex;
-}
-
-EProjectPackagingBuildConfigurations UProjectPackagingSettings::GetBuildConfigurationForPlatform(FName PlatformName) const
-{
-	const EProjectPackagingBuildConfigurations* Value = PerPlatformBuildConfig.Find(PlatformName);
-
-	return Value == nullptr ? EProjectPackagingBuildConfigurations::PPBC_Development : *Value;
-}
-
-void UProjectPackagingSettings::SetBuildConfigurationForPlatform(FName PlatformName, EProjectPackagingBuildConfigurations Configuration)
-{
-	PerPlatformBuildConfig.Add(PlatformName, Configuration);
-}
-
-FName UProjectPackagingSettings::GetTargetPlatformForPlatform(FName PlatformName) const
-{
-	const FName* Value = PerPlatformTargetPlatformName.Find(PlatformName);
-
-	// the platform name is also the name of the vanilla info
-	return Value == nullptr ? PlatformName : *Value;
-}
-
-void UProjectPackagingSettings::SetTargetPlatformForPlatform(FName PlatformName, FName TargetPlatformName)
-{
-	PerPlatformTargetPlatformName.Add(PlatformName, TargetPlatformName);
-}
-
 
 /* UCrashReporterSettings interface
 *****************************************************************************/
diff --git a/Engine/Source/Editor/UnrealEd/Public/CookerSettings.h b/Engine/Source/Editor/UnrealEd/Public/CookerSettings.h
deleted file mode 100644
diff --git a/Engine/Source/Runtime/CoreUObject/Public/UObject/NoExportTypes.h b/Engine/Source/Runtime/CoreUObject/Public/UObject/NoExportTypes.h
--- a/Engine/Source/Runtime/CoreUObject/Public/UObject/NoExportTypes.h
+++ b/Engine/Source/Runtime/CoreUObject/Public/UObject/NoExportTypes.h
@@ -1657,6 +1657,34 @@ public:
 	void ExecuteUbergraph(int32 EntryPoint);
 };
 
+/** Structure for file paths that are displayed in the editor with a picker UI. */
+USTRUCT(noexport, BlueprintType)
+struct FFilePath
+{
+	GENERATED_BODY()
+
+	/**
+		* The path to the file.
+		*/
+	UPROPERTY(EditAnywhere, BlueprintReadWrite, Category = FilePath)
+	FString FilePath;
+};
+
+/** Structure for directory paths that are displayed in the editor with a picker UI. */
+USTRUCT(noexport, BlueprintType)
+struct FDirectoryPath
+{
+	GENERATED_BODY()
+
+	/**
+		* The path to the directory.
+		*/
+	UPROPERTY(EditAnywhere, BlueprintReadWrite, Category = Path)
+	FString Path;
+};
+
+
+
 /// @endcond
 
 #endif
diff --git a/Engine/Source/Runtime/CoreUObject/Public/UObject/SoftObjectPath.h b/Engine/Source/Runtime/CoreUObject/Public/UObject/SoftObjectPath.h
--- a/Engine/Source/Runtime/CoreUObject/Public/UObject/SoftObjectPath.h
+++ b/Engine/Source/Runtime/CoreUObject/Public/UObject/SoftObjectPath.h
@@ -425,4 +425,22 @@ struct FSoftObjectPathSerializationScope
 	{
 		FSoftObjectPathThreadContext::Get().OptionStack.Pop();
 	}
-};
\ No newline at end of file
+};
+
+/** Structure for file paths that are displayed in the editor with a picker UI. */
+struct FFilePath
+{
+	/**
+	 * The path to the file.
+	 */
+	FString FilePath;
+};
+
+/** Structure for directory paths that are displayed in the editor with a picker UI. */
+struct FDirectoryPath
+{
+	/**
+	 * The path to the directory.
+	 */
+	FString Path;
+};
diff --git a/Engine/Source/Runtime/Engine/Classes/Engine/EngineTypes.h b/Engine/Source/Runtime/Engine/Classes/Engine/EngineTypes.h
--- a/Engine/Source/Runtime/Engine/Classes/Engine/EngineTypes.h
+++ b/Engine/Source/Runtime/Engine/Classes/Engine/EngineTypes.h
@@ -3875,32 +3875,6 @@ struct FComponentSocketDescription
 /** Dynamic delegate to use by components that want to route the broken-event into blueprints */
 DECLARE_DYNAMIC_MULTICAST_DELEGATE_OneParam(FConstraintBrokenSignature, int32, ConstraintIndex);
 
-/** Structure for file paths that are displayed in the editor with a picker UI. */
-USTRUCT(BlueprintType)
-struct FFilePath
-{
-	GENERATED_BODY()
-
-	/**
-	 * The path to the file.
-	 */
-	UPROPERTY(EditAnywhere, BlueprintReadWrite, Category=FilePath)
-	FString FilePath;
-};
-
-/** Structure for directory paths that are displayed in the editor with a picker UI. */
-USTRUCT(BlueprintType)
-struct FDirectoryPath
-{
-	GENERATED_BODY()
-
-	/**
-	 * The path to the directory.
-	 */
-	UPROPERTY(EditAnywhere, BlueprintReadWrite, Category=Path)
-	FString Path;
-};
-
 
 /**
  * Reference to an editor collection of assets. This allows an editor-only picker UI
diff --git a/Engine/Source/Runtime/Engine/Engine.Build.cs b/Engine/Source/Runtime/Engine/Engine.Build.cs
--- a/Engine/Source/Runtime/Engine/Engine.Build.cs
+++ b/Engine/Source/Runtime/Engine/Engine.Build.cs
@@ -317,10 +317,11 @@ public class Engine : ModuleRules
 
 		if (!Target.bBuildRequiresCookedData && Target.bCompileAgainstEngine)
 		{
+			PublicDependencyModuleNames.Add("TargetPlatform");
+
 			DynamicallyLoadedModuleNames.AddRange(
 				new string[] {
 					"DerivedDataCache",
-					"TargetPlatform",
 					"DesktopPlatform"
 				}
 			);
diff --git a/Engine/Source/Runtime/Engine/Private/Blueprint.cpp b/Engine/Source/Runtime/Engine/Private/Blueprint.cpp
--- a/Engine/Source/Runtime/Engine/Private/Blueprint.cpp
+++ b/Engine/Source/Runtime/Engine/Private/Blueprint.cpp
@@ -21,7 +21,7 @@
 
 #if WITH_EDITOR
 #include "BlueprintCompilationManager.h"
-#include "Editor/UnrealEd/Classes/Settings/ProjectPackagingSettings.h"
+#include "TargetPlatform/Classes/Settings/ProjectPackagingSettings.h"
 #include "Editor/UnrealEd/Classes/Settings/EditorExperimentalSettings.h"
 #include "Engine/SimpleConstructionScript.h"
 #include "Engine/SCS_Node.h"
