- Made lack of UBT during AutoSDK setup a warning instead of Fatal

diff --git a/Engine/Source/Developer/TargetPlatform/Private/TargetPlatformManagerModule.cpp b/Engine/Source/Developer/TargetPlatform/Private/TargetPlatformManagerModule.cpp
--- a/Engine/Source/Developer/TargetPlatform/Private/TargetPlatformManagerModule.cpp
+++ b/Engine/Source/Developer/TargetPlatform/Private/TargetPlatformManagerModule.cpp
@@ -108,7 +108,7 @@ public:
 				FString UBTOutput;
 				if (!FDesktopPlatformModule::Get()->InvokeUnrealBuildToolSync(UBTParams, *GLog, true, UBTReturnCode, UBTOutput))
 				{
-					UE_LOG(LogTargetPlatformManager, Fatal, TEXT("Failed to run UBT to check SDK status!"));
+					UE_LOG(LogTargetPlatformManager, Warning, TEXT("AutoSDK is enabled (UE_SDKS_ROOT is set), but failed to run UBT to check SDK status! Check your installation."));
 				}
 			}
 
@@ -661,7 +661,7 @@ protected:
 		FString SDKPath = FPlatformMisc::GetEnvironmentVariable(*SDKRootEnvFar);
 
 		// AutoSDKs only enabled if UE_SDKS_ROOT is set.
-		if (SDKPath.Len() != 0)
+//		if (SDKPath.Len() != 0)
 		{
 			return true;
 		}
