- Restored a line of code I mistakenly commented out

diff --git a/Engine/Source/Developer/TargetPlatform/Private/TargetPlatformManagerModule.cpp b/Engine/Source/Developer/TargetPlatform/Private/TargetPlatformManagerModule.cpp
--- a/Engine/Source/Developer/TargetPlatform/Private/TargetPlatformManagerModule.cpp
+++ b/Engine/Source/Developer/TargetPlatform/Private/TargetPlatformManagerModule.cpp
@@ -661,7 +661,7 @@ protected:
 		FString SDKPath = FPlatformMisc::GetEnvironmentVariable(*SDKRootEnvFar);
 
 		// AutoSDKs only enabled if UE_SDKS_ROOT is set.
-//		if (SDKPath.Len() != 0)
+		if (SDKPath.Len() != 0)
 		{
 			return true;
 		}
