adding comments to help the next person who stumbles through here like me.

diff --git a/Engine/Source/Developer/TargetPlatform/Public/Interfaces/IAudioFormat.h b/Engine/Source/Developer/TargetPlatform/Public/Interfaces/IAudioFormat.h
--- a/Engine/Source/Developer/TargetPlatform/Public/Interfaces/IAudioFormat.h
+++ b/Engine/Source/Developer/TargetPlatform/Public/Interfaces/IAudioFormat.h
@@ -56,8 +56,8 @@ public:
 	 * Cooks the source data for the platform and stores the cooked data internally.
 	 *
 	 * @param Format The desired format.
-	 * @param SrcBuffer The source buffer.
-	 * @param QualityInfo All the information the compressor needs to compress the audio.
+	 * @param SrcBuffer The source buffer. Buffers are 16 bit PCM, either mono or stereo (check QualityInfo.NumChannels)
+	 * @param QualityInfo All the information the compressor needs to compress the audio. QualityInfo.Duration is unset.
 	 * @param OutBuffer Will hold the resulting compressed audio.
 	 * @return true on success, false otherwise.
 	 */
@@ -67,8 +67,8 @@ public:
 	 * Cooks up to 8 mono files into a multi-stream file (e.g. 5.1). The front left channel is required, the rest are optional.
 	 *
 	 * @param Format The desired format.
-	 * @param SrcBuffers The source buffers.
-	 * @param QualityInfo All the information the compressor needs to compress the audio
+	 * @param SrcBuffers The source buffers. Buffers are mono 16 bit PCM.
+	 * @param QualityInfo All the information the compressor needs to compress the audio. QualityInfo.Duration is unset.
 	 * @param OutBuffer Will contain the resulting compressed audio.
 	 * @return true if succeeded, false otherwise
 	 */
@@ -92,6 +92,8 @@ public:
 	/** 
 	 * Re-compresses raw PCM to the the platform dependent format, and then back to PCM. Used for quality previewing.
 	 *
+	 * This function is, as far as I can tell, unused.
+	 * 
 	 * @param Format The desired format.
 	 * @param SrcBuffer Uncompressed PCM data.
 	 * @param QualityInfo All the information the compressor needs to compress the audio.
@@ -103,6 +105,8 @@ public:
 	/**
 	 * Given the encoded buffer, returns the minimum number of bytes required to perform ICompressedAudioInfo::ReadCompressedInfo() or IStreamedCompressedInfo::ParseHeader() for this file.
 	 * 
+	 * This is called prior to SplitDataForStreaming(), and is used for determining FirstChunkMaxSize.
+	 * 
 	 * @param Format the codec that SrcBuffer was compressed as.
 	 * @param SrcBuffer the compressed data that will later be split into individual chunks.
 	 */
