- Allow for each TargetPlatform class to set what DeviceProfile it uses while cooking (with this checkin: desktop = PlatformName, non-desktop = IniPlatformName)

(From Josh.Adams)

diff --git a/Engine/Source/Developer/TargetPlatform/Public/Common/TargetPlatformBase.h b/Engine/Source/Developer/TargetPlatform/Public/Common/TargetPlatformBase.h
--- a/Engine/Source/Developer/TargetPlatform/Public/Common/TargetPlatformBase.h
+++ b/Engine/Source/Developer/TargetPlatform/Public/Common/TargetPlatformBase.h
@@ -307,6 +307,12 @@ public:
 		return FString(TPlatformProperties::IniPlatformName());
 	}
 
+	virtual FString CookingDeviceProfileName() const override
+	{
+		// default to using TargetPlatform name, like WindowsClient, when looking up the DeviceProfile to cook with
+		return PlatformName();
+	}
+
 	virtual bool RequiresCookedData() const override
 	{
 		return TPlatformProperties::RequiresCookedData();
@@ -465,6 +471,12 @@ public:
 		return this->PlatformInfo->IniPlatformName.ToString();
 	}
 
+	virtual FString CookingDeviceProfileName() const override
+	{
+		// when cooking for non-desktop platforms, always use the base platform name as the DP to cook with (there aren't IOSClient DPs, there are IOS (base), IPhone7, iPad4, etc DPs)
+		return IniPlatformName();
+	}
+
 	virtual bool HasEditorOnlyData() const override
 	{
 		return false;
diff --git a/Engine/Source/Developer/TargetPlatform/Public/Interfaces/ITargetPlatform.h b/Engine/Source/Developer/TargetPlatform/Public/Interfaces/ITargetPlatform.h
--- a/Engine/Source/Developer/TargetPlatform/Public/Interfaces/ITargetPlatform.h
+++ b/Engine/Source/Developer/TargetPlatform/Public/Interfaces/ITargetPlatform.h
@@ -240,6 +240,11 @@ public:
 	 */
 	virtual FString IniPlatformName() const = 0;
 
+	/**
+	 * Gets the name of the device profile to use when cooking this TargetPlatform
+	 */
+	virtual FString CookingDeviceProfileName() const = 0;
+
 	/**
 	 * Enables/Disable the device check
 	 */
diff --git a/Engine/Source/Runtime/Engine/Private/DeviceProfiles/DeviceProfileManager.cpp b/Engine/Source/Runtime/Engine/Private/DeviceProfiles/DeviceProfileManager.cpp
--- a/Engine/Source/Runtime/Engine/Private/DeviceProfiles/DeviceProfileManager.cpp
+++ b/Engine/Source/Runtime/Engine/Private/DeviceProfiles/DeviceProfileManager.cpp
@@ -773,7 +773,9 @@ void UDeviceProfileManager::LoadProfiles()
 				ITargetPlatform* Platform = TargetPlatforms[PlatformIndex];
 
 				// Set TextureLODSettings
-				const UTextureLODSettings* TextureLODSettingsObj = FindProfile(Platform->PlatformName(), false);
+				const UTextureLODSettings* TextureLODSettingsObj = FindProfile(Platform->CookingDeviceProfileName(), false);
+				checkf(TextureLODSettingsObj, TEXT("No TextureLODSettings found for %s"), *Platform->CookingDeviceProfileName());
+
 				Platform->RegisterTextureLODSettings(TextureLODSettingsObj);
 			}
 		}
