Undo //UE5/Main/Engine/Source/... changelist 16165786

diff --git a/Engine/Source/Developer/TargetPlatform/Private/TargetPlatformManagerModule.cpp b/Engine/Source/Developer/TargetPlatform/Private/TargetPlatformManagerModule.cpp
--- a/Engine/Source/Developer/TargetPlatform/Private/TargetPlatformManagerModule.cpp
+++ b/Engine/Source/Developer/TargetPlatform/Private/TargetPlatformManagerModule.cpp
@@ -17,7 +17,6 @@
 #include "Interfaces/IShaderFormat.h"
 #include "Interfaces/IShaderFormatModule.h"
 #include "Interfaces/ITextureFormat.h"
-#include "Interfaces/ITextureFormatManagerModule.h"
 #include "Interfaces/ITextureFormatModule.h"
 #include "PlatformInfo.h"
 #include "DesktopPlatformModule.h"
@@ -124,8 +123,6 @@ public:
 			}
 		}
 #endif
-		TextureFormatManager = FModuleManager::LoadModulePtr<ITextureFormatManagerModule>("TextureFormat");
-
 		// Calling a virtual function from a constructor, but with no expectation that a derived implementation of this
 		// method would be called.  This is solely to avoid duplicating code in this implementation, not for polymorphism.
 		FTargetPlatformManagerModule::Invalidate();
@@ -164,7 +161,7 @@ public:
 		if (!bHasInitErrors)
 		{
 			GetAudioFormats();
-			TextureFormatManager->Invalidate();
+			GetTextureFormats();
 			GetShaderFormats();
 		}
 
@@ -456,12 +453,60 @@ public:
 
 	virtual const TArray<const ITextureFormat*>& GetTextureFormats() override
 	{
-		return TextureFormatManager->GetTextureFormats();
+		static bool bInitialized = false;
+		static TArray<const ITextureFormat*> Results;
+
+		if (!bInitialized || bForceCacheUpdate)
+		{
+			bInitialized = true;
+			Results.Empty(Results.Num());
+
+			TArray<FName> Modules;
+
+			FModuleManager::Get().FindModules(TEXT("*TextureFormat*"), Modules);
+
+			if (!Modules.Num())
+			{
+				UE_LOG(LogTargetPlatformManager, Error, TEXT("No target texture formats found!"));
+			}
+
+			for (int32 Index = 0; Index < Modules.Num(); Index++)
+			{
+				ITextureFormatModule* Module = FModuleManager::LoadModulePtr<ITextureFormatModule>(Modules[Index]);
+				if (Module)
+				{
+					ITextureFormat* Format = Module->GetTextureFormat();
+					if (Format != nullptr)
+					{
+						Results.Add(Format);
+					}
+				}
+			}
+		}
+
+		return Results;
 	}
 
 	virtual const ITextureFormat* FindTextureFormat(FName Name) override
 	{
-		return TextureFormatManager->FindTextureFormat(Name);
+		const TArray<const ITextureFormat*>& TextureFormats = GetTextureFormats();
+
+		for (int32 Index = 0; Index < TextureFormats.Num(); Index++)
+		{
+			TArray<FName> Formats;
+
+			TextureFormats[Index]->GetSupportedFormats(Formats);
+
+			for (int32 FormatIndex = 0; FormatIndex < Formats.Num(); FormatIndex++)
+			{
+				if (Formats[FormatIndex] == Name)
+				{
+					return TextureFormats[Index];
+				}
+			}
+		}
+
+		return nullptr;
 	}
 
 	virtual const TArray<const IShaderFormat*>& GetShaderFormats() override
@@ -1162,9 +1207,6 @@ private:
 	// Map for fast lookup of platforms by name.
 	TMap<FName, ITargetPlatform*> PlatformsByName;
 
-	// External module that texture format operations are forwarded to
-	ITextureFormatManagerModule* TextureFormatManager;
-
 #if AUTOSDKS_ENABLED
 	// holds the list of Platforms that have attempted setup.
 	TMap<FName, bool> PlatformsSetup;
diff --git a/Engine/Source/Developer/TargetPlatform/Public/ChildTextureFormat.h b/Engine/Source/Developer/TargetPlatform/Public/ChildTextureFormat.h
new file mode 100644
--- /dev/null
+++ b/Engine/Source/Developer/TargetPlatform/Public/ChildTextureFormat.h
@@ -0,0 +1,197 @@
+// Copyright Epic Games, Inc. All Rights Reserved.
+
+#pragma once
+
+#include "Interfaces/ITextureFormat.h"
+#include "Interfaces/ITextureFormatModule.h"
+#include "Interfaces/ITargetPlatformManagerModule.h"
+#include "TextureCompressorModule.h"
+
+/**
+ * Version of ITextureFormat that handles a child texture format that is used as a "post-process" after compressing textures, useful for
+ * several platforms that need to modify already compressed texture data for optimal data
+ */
+class FChildTextureFormat : public ITextureFormat
+{
+public:
+
+	FChildTextureFormat(const TCHAR* PlatformFormatPrefix)
+		: FormatPrefix(PlatformFormatPrefix)
+	{
+
+	}
+
+protected:
+	void AddBaseTextureFormatModules(const TCHAR* ModuleNameWildcard)
+	{
+		TArray<FName> Modules;
+		FModuleManager::Get().FindModules(ModuleNameWildcard, Modules);
+
+		for (FName ModuleName : Modules)
+		{
+			ITextureFormatModule * TFModule = FModuleManager::LoadModulePtr<ITextureFormatModule>(ModuleName);
+			if ( TFModule != nullptr )
+			{
+				ITextureFormat* BaseFormat = TFModule->GetTextureFormat();
+				BaseFormat->GetSupportedFormats(BaseFormats);
+			}
+		}
+	}
+
+	FName GetBaseFormatName(FName PlatformName) const
+	{
+		return FName(*(PlatformName.ToString().Replace(*FormatPrefix, TEXT(""))));
+	}
+
+	/**
+	 * Given a platform specific format name, get the parent texture format object
+	 */
+	const ITextureFormat* GetBaseFormatObject(FName PlatformName) const
+	{
+		FName BaseFormatName = GetBaseFormatName(PlatformName);
+
+		ITargetPlatformManagerModule& TPM = FModuleManager::LoadModuleChecked<ITargetPlatformManagerModule>("TargetPlatform");
+		const ITextureFormat* FormatObject = TPM.FindTextureFormat(BaseFormatName);
+
+		checkf(FormatObject != nullptr, TEXT("Bad PlatformName %s passed to FChildTextureFormat::GetBaseFormatObject()"));
+
+		return FormatObject;
+	}
+
+	/**
+	 * The final version is a combination of parent and child formats, 8 bits for each
+	 */
+	virtual uint8 GetChildFormatVersion(FName Format, const struct FTextureBuildSettings* BuildSettings) const = 0;
+
+	/**
+	 * Make the child type think about if they need a key string or not, by making it pure virtual
+	 */
+	virtual FString GetChildDerivedDataKeyString(const class UTexture& Texture, const FTextureBuildSettings* BuildSettings) const = 0;
+
+
+public:
+
+	//// ITextureFormat interface ////
+
+	virtual void GetSupportedFormats(TArray<FName>& OutFormats) const override
+	{
+		for (FName BaseFormat : BaseFormats)
+		{
+			FName ChildFormat(*(FormatPrefix + BaseFormat.ToString()));
+			OutFormats.Add(ChildFormat);
+		}
+	}
+
+	virtual uint16 GetVersion(FName Format, const struct FTextureBuildSettings* BuildSettings) const final
+	{
+		uint16 BaseVersion = GetBaseFormatObject(Format)->GetVersion(Format, BuildSettings);
+		checkf(BaseVersion < 256, TEXT("BaseFormat for %s had too large a version (%d), must fit in 8bits"), *Format.ToString(), BaseVersion);
+
+		uint8 ChildVersion = GetChildFormatVersion(Format, BuildSettings);
+
+		// 8 bits for each version
+		return (BaseVersion << 8) | ChildVersion;
+	}
+
+	virtual FString GetDerivedDataKeyString(const class UTexture& Texture, const FTextureBuildSettings* BuildSettings) const final
+	{
+		FTextureBuildSettings BaseSettings = *BuildSettings;
+		BaseSettings.TextureFormatName = GetBaseFormatName(BuildSettings->TextureFormatName);
+
+		FString BaseString = GetBaseFormatObject(BuildSettings->TextureFormatName)->GetDerivedDataKeyString(Texture, &BaseSettings);
+		FString ChildString = GetChildDerivedDataKeyString(Texture, BuildSettings);
+
+		return BaseString + ChildString;
+	}
+
+	virtual EPixelFormat GetPixelFormatForImage(const struct FTextureBuildSettings& BuildSettings, const struct FImage& ExampleImage, bool bImageHasAlphaChannel) const override 
+	{
+		FTextureBuildSettings Settings = BuildSettings;
+		Settings.TextureFormatName = GetBaseFormatName(BuildSettings.TextureFormatName);
+		return GetBaseFormatObject(BuildSettings.TextureFormatName)->GetPixelFormatForImage(Settings, ExampleImage, bImageHasAlphaChannel);
+	}
+
+	bool CompressBaseImage(
+		const FImage& InImage,
+		const struct FTextureBuildSettings& BuildSettings,
+		bool bImageHasAlphaChannel,
+		FCompressedImage2D& OutCompressedImage
+	) const
+	{
+		FTextureBuildSettings BaseSettings = BuildSettings;
+		BaseSettings.TextureFormatName = GetBaseFormatName(BuildSettings.TextureFormatName);
+
+		// pass along the compression to the base format
+		if (GetBaseFormatObject(BuildSettings.TextureFormatName)->CompressImage(InImage, BaseSettings, bImageHasAlphaChannel, OutCompressedImage) == false)
+		{
+			UE_LOG(LogTemp, Error, TEXT("Failed to compress with base compressor [format %s]"), *BaseSettings.TextureFormatName.ToString());
+			return false;
+		}
+		return true;
+	}
+
+	bool CompressBaseImageTiled(
+		const FImage* Images,
+		uint32 NumImages,
+		const struct FTextureBuildSettings& BuildSettings,
+		bool bImageHasAlphaChannel,
+		TSharedPtr<FTilerSettings>& TilerSettings,
+		FCompressedImage2D& OutCompressedImage
+	) const
+	{
+		FTextureBuildSettings BaseSettings = BuildSettings;
+		BaseSettings.TextureFormatName = GetBaseFormatName(BuildSettings.TextureFormatName);
+
+		// pass along the compression to the base format
+		if (GetBaseFormatObject(BuildSettings.TextureFormatName)->CompressImageTiled(Images, NumImages, BaseSettings, bImageHasAlphaChannel, TilerSettings, OutCompressedImage) == false)
+		{
+			UE_LOG(LogTemp, Error, TEXT("Failed to compress with base tiled compressor [format %s]"), *BaseSettings.TextureFormatName.ToString());
+			return false;
+		}
+		return true;
+	}
+
+	bool PrepareTiling(
+		const FImage* Images,
+		const uint32 NumImages,
+		const struct FTextureBuildSettings& BuildSettings,
+		bool bImageHasAlphaChannel,
+		TSharedPtr<FTilerSettings>& OutTilerSettings,
+		TArray<FCompressedImage2D>& OutCompressedImages
+	) const override
+	{
+		FTextureBuildSettings BaseSettings = BuildSettings;
+		BaseSettings.TextureFormatName = GetBaseFormatName(BuildSettings.TextureFormatName);
+
+		return GetBaseFormatObject(BuildSettings.TextureFormatName)->PrepareTiling(Images, NumImages, BaseSettings, bImageHasAlphaChannel, OutTilerSettings, OutCompressedImages);
+	}
+
+	bool SetTiling(
+		const struct FTextureBuildSettings& BuildSettings,
+		TSharedPtr<FTilerSettings>& TilerSettings,
+		const TArray64<uint8>& ReorderedBlocks,
+		uint32 NumBlocks
+	) const override
+	{
+		FTextureBuildSettings BaseSettings = BuildSettings;
+		BaseSettings.TextureFormatName = GetBaseFormatName(BuildSettings.TextureFormatName);
+
+		return GetBaseFormatObject(BuildSettings.TextureFormatName)->SetTiling(BaseSettings, TilerSettings, ReorderedBlocks, NumBlocks);
+	}
+
+	void ReleaseTiling(const struct FTextureBuildSettings& BuildSettings, TSharedPtr<FTilerSettings>& TilerSettings) const override
+	{
+		FTextureBuildSettings BaseSettings = BuildSettings;
+		BaseSettings.TextureFormatName = GetBaseFormatName(BuildSettings.TextureFormatName);
+
+		return GetBaseFormatObject(BuildSettings.TextureFormatName)->ReleaseTiling(BuildSettings, TilerSettings);
+	}
+
+protected:
+
+	// Prefix put before all formats from parent formats
+	FString FormatPrefix;
+
+	// List of base formats that. Combined with FormatPrefix, this contains all formats this can handle
+	TArray<FName> BaseFormats;
+};
diff --git a/Engine/Source/Developer/TargetPlatform/Public/Interfaces/ITextureFormat.h b/Engine/Source/Developer/TargetPlatform/Public/Interfaces/ITextureFormat.h
new file mode 100644
--- /dev/null
+++ b/Engine/Source/Developer/TargetPlatform/Public/Interfaces/ITextureFormat.h
@@ -0,0 +1,245 @@
+// Copyright Epic Games, Inc. All Rights Reserved.
+
+#pragma once
+
+#include "CoreMinimal.h"
+#include "PixelFormat.h"
+/**
+ * Structure for texture format compressor capabilities.
+ */
+struct FTextureFormatCompressorCaps
+{
+	FTextureFormatCompressorCaps()
+		: MaxTextureDimension(TNumericLimits<uint32>::Max())
+		, NumMipsInTail(0)
+		, ExtData(0)
+	{ }
+
+	uint32 MaxTextureDimension;
+	uint32 NumMipsInTail;
+	uint32 ExtData;
+};
+
+
+/**
+ * Interface for texture compression modules.
+ */
+class ITextureFormat
+{
+public:
+
+	/**
+	 * Checks whether this texture format can compress in parallel.
+	 *
+	 * @return true if parallel compression is supported, false otherwise.
+	 */
+	virtual bool AllowParallelBuild() const
+	{
+		return false;
+	}
+
+	/**
+	 * Gets the current version of the specified texture format.
+	 *
+	 * @param Format The format to get the version for.
+	 * @return Version number.
+	 */
+	virtual uint16 GetVersion(
+		FName Format,
+		const struct FTextureBuildSettings* BuildSettings = nullptr
+	) const = 0;
+
+	/**
+	 * Gets an optional derived data key string, so that the compressor can
+	 * rely upon the number of mips, size of texture, etc, when compressing the image
+	 *
+	 * @param Texture Reference to the texture we are compressing.
+	 * @return A string that will be used with the DDC, the string should be in the format "<DATA>_"
+	 */
+	virtual FString GetDerivedDataKeyString( const class UTexture& Texture, const FTextureBuildSettings* BuildSettings) const
+	{
+		return TEXT("");
+	}
+
+	/**
+	 * Gets the list of supported formats.
+	 *
+	 * @param OutFormats Will hold the list of formats.
+	 */
+	virtual void GetSupportedFormats( TArray<FName>& OutFormats ) const = 0;
+
+	/**
+	* Gets the capabilities of the texture compressor.
+	*
+	* @param OutCaps Filled with capability properties of texture format compressor.
+	*/
+	virtual FTextureFormatCompressorCaps GetFormatCapabilities() const = 0;
+
+	/**
+	* Gets the capabilities of the texture compressor.
+	*
+	* @param OutCaps Filled with capability properties of texture format compressor.
+	*/
+	virtual FTextureFormatCompressorCaps GetFormatCapabilitiesEx(const struct FTextureBuildSettings& BuildSettings, uint32 NumMips, const struct FImage& ExampleImage, bool bImageHasAlphaChannel) const
+	{
+		return GetFormatCapabilities();
+	}
+
+	/**
+	 * Calculate the final/runtime pixel format for this image on this platform
+	 */
+	virtual EPixelFormat GetPixelFormatForImage(const struct FTextureBuildSettings& BuildSettings, const struct FImage& Image, bool bImageHasAlphaChannel) const = 0;
+
+	/**
+	 * Compresses a single image.
+	 *
+	 * @param Image The input image.
+	 * @param BuildSettings Build settings.
+	 * @param bImageHasAlphaChannel true if the image has a non-white alpha channel.
+	 * @param OutCompressedMip The compressed image.
+	 * @returns true on success, false otherwise.
+	 */
+	virtual bool CompressImage(
+		const struct FImage& Image,
+		const struct FTextureBuildSettings& BuildSettings,
+		bool bImageHasAlphaChannel,
+		struct FCompressedImage2D& OutCompressedImage
+	) const = 0;
+
+	/**
+	 * Compress an image (or images for a miptail) into a single mip blob.
+	 *
+	 * @param Images The input image(s)
+	 * @param Images The number of images (for a miptail, this number should match what was returned in GetFormatCapabilities, mostly used for verification)
+	 * @param BuildSettings Build settings.
+	 * @param bImageHasAlphaChannel true if the image has a non-white alpha channel.
+	 * @param ExtData Extra data that the format may want to have passed back in to each compress call (makes the format class be stateless)
+	 * @param OutCompressedMip The compressed image.
+	 * @returns true on success, false otherwise.
+	 */
+	virtual bool CompressImageEx(const struct FImage* Images,
+		const uint32 NumImages,
+		const struct FTextureBuildSettings& BuildSettings,
+		bool bImageHasAlphaChannel,
+		uint32 ExtData,
+		FCompressedImage2D& OutCompressedImage) const
+	{
+		// general case can't handle mip tails
+		if (Images == nullptr || NumImages > 1)
+		{
+			return false;
+		}
+		
+		return CompressImage(*Images, BuildSettings, bImageHasAlphaChannel, OutCompressedImage);
+	}
+
+	/**
+	 * An object produced by PrepareTiling and used by SetTiling and CompressImageTiled.
+	 */
+	struct FTilerSettings
+	{
+	};
+
+	/**
+	 * Compress an image (or images for a miptail) into a single mip blob with device-specific tiling.
+	 *
+	 * @param Image The input image.
+	 * @param BuildSettings Build settings.
+	 * @param bImageHasAlphaChannel true if the image has a non-white alpha channel.
+	 * @param OutCompressedMip The compressed image.
+	 * @param Tiler The tiler settings.
+	 * @returns true on success, false otherwise.
+	 */
+	virtual bool CompressImageTiled(
+		const struct FImage* Images,
+		uint32 NumImages,
+		const struct FTextureBuildSettings& BuildSettings,
+		bool bImageHasAlphaChannel,
+		TSharedPtr<FTilerSettings>& TilerSettings,
+		struct FCompressedImage2D& OutCompressedImage) const
+	{
+		unimplemented();
+		return false;
+	}
+
+	/**
+	 * Whether device-specific tiling is supported by the compressor.
+	 *
+	 * @param BuildSettings Build settings.
+	 * @returns true if tiling is supported, false if it must be done by the caller
+	 *
+	 */
+	virtual bool SupportsTiling(const struct FTextureBuildSettings& BuildSettings) const
+	{
+		return false;
+	}
+
+	/**
+	 * Prepares to compresses a single image with tiling. The result OutTilerSettings is used by SetTiling and CompressImageTiled.
+	 *
+	 * @param Image The input image.
+	 * @param BuildSettings Build settings.
+	 * @param bImageHasAlphaChannel true if the image has a non-white alpha channel.
+	 * @param OutTilerSettings The tiler settings that will be used by CompressImageTiled and SetTiling.
+	 * @param OutCompressedImage The image to tile.
+	 * @returns true on success, false otherwise.
+	 */
+	virtual bool PrepareTiling(
+		const FImage* Images,
+		const uint32 NumImages,
+		const struct FTextureBuildSettings& BuildSettings,
+		bool bImageHasAlphaChannel,
+		TSharedPtr<FTilerSettings>& OutTilerSettings,
+		TArray<FCompressedImage2D>& OutCompressedImage
+	) const
+	{
+		unimplemented();
+		return false;
+	}
+	
+	/**
+	 * Sets the tiling settings after device-specific tiling has been performed.
+	 *
+	 * @param BuildSettings Build settings.
+	 * @param TilerSettings The tiler settings produced by PrepareTiling.
+	 * @param ReorderedBlocks The blocks that have been tiled.
+	 * @param NumBlocks The number of blocks.
+	 * @returns true on success, false otherwise.
+	 */
+	virtual bool SetTiling(
+		const struct FTextureBuildSettings& BuildSettings,
+		TSharedPtr<FTilerSettings>& TilerSettings,
+		const TArray64<uint8>& ReorderedBlocks,
+		uint32 NumBlocks
+	) const
+	{
+		unimplemented();
+		return false;
+	}
+
+	/**
+	 * Cleans up the FTilerSettings object once it is finished.
+	 * 
+	 * @param BuildSettings Build settings.
+	 * @param TilerSettings The tiler settings object to release.
+	 */
+	virtual void ReleaseTiling(const struct FTextureBuildSettings& BuildSettings, TSharedPtr<FTilerSettings>& TilerSettings) const
+	{
+		unimplemented();
+	}
+
+	/**
+	 * Whether the compressor uses the FTaskGraph API.
+	 * 
+	 * @returns true if FTaskGraph is used, false otherwise
+	 */
+	virtual bool UsesTaskGraph() const
+	{
+		return false;
+	}
+
+public:
+
+	/** Virtual destructor. */
+	virtual ~ITextureFormat() { }
+};
diff --git a/Engine/Source/Developer/TargetPlatform/Public/Interfaces/ITextureFormatModule.h b/Engine/Source/Developer/TargetPlatform/Public/Interfaces/ITextureFormatModule.h
new file mode 100644
--- /dev/null
+++ b/Engine/Source/Developer/TargetPlatform/Public/Interfaces/ITextureFormatModule.h
@@ -0,0 +1,29 @@
+// Copyright Epic Games, Inc. All Rights Reserved.
+
+#pragma once
+
+#include "CoreMinimal.h"
+#include "Modules/ModuleInterface.h"
+
+class ITextureFormat;
+
+/**
+ * Interface for texture format modules.
+ */
+class ITextureFormatModule
+	: public IModuleInterface
+{
+public:
+
+	/**
+	 * Gets the texture format.
+	 *
+	 * @return The texture format interface.
+	 */
+	virtual ITextureFormat* GetTextureFormat() = 0;
+
+public:
+
+	/** Virtual destructor. */
+	~ITextureFormatModule() { }
+};
diff --git a/Engine/Source/Developer/TargetPlatform/TargetPlatform.Build.cs b/Engine/Source/Developer/TargetPlatform/TargetPlatform.Build.cs
--- a/Engine/Source/Developer/TargetPlatform/TargetPlatform.Build.cs
+++ b/Engine/Source/Developer/TargetPlatform/TargetPlatform.Build.cs
@@ -18,11 +18,6 @@ public class TargetPlatform : ModuleRules
 		PublicDependencyModuleNames.Add("DesktopPlatform");
 		PublicDependencyModuleNames.Add("LauncherPlatform");
 
-		// TextureFormat contains public headers that were historically part of TargetPlatform, so it is exposed
-		// as a public include path on TargetPlatform.
-		PublicIncludePathModuleNames.Add("TextureFormat");
-		PublicDependencyModuleNames.Add("TextureFormat");
-
 		PrivateIncludePathModuleNames.Add("Engine");
 		PrivateIncludePathModuleNames.Add("PhysicsCore");
 
diff --git a/Engine/Source/Developer/TextureCompressor/Private/TextureCompressorModule.cpp b/Engine/Source/Developer/TextureCompressor/Private/TextureCompressorModule.cpp
--- a/Engine/Source/Developer/TextureCompressor/Private/TextureCompressorModule.cpp
+++ b/Engine/Source/Developer/TextureCompressor/Private/TextureCompressorModule.cpp
@@ -7,8 +7,8 @@
 #include "Async/AsyncWork.h"
 #include "Async/ParallelFor.h"
 #include "Modules/ModuleManager.h"
-#include "Engine/TextureDefines.h"
-#include "TextureFormatManager.h"
+#include "Engine/Texture.h"
+#include "Interfaces/ITargetPlatformManagerModule.h"
 #include "Interfaces/ITextureFormat.h"
 #include "Misc/Paths.h"
 #include "ImageCore.h"
@@ -2202,9 +2202,10 @@ public:
 	virtual bool UsesTaskGraph(const FTextureBuildSettings& BuildSettings) const override
 	{
 		const ITextureFormat* TextureFormat = nullptr;
-		if (TFM)
+		ITargetPlatformManagerModule* TPM = GetTargetPlatformManager();
+		if (TPM)
 		{
-			TextureFormat = TFM->FindTextureFormat(BuildSettings.TextureFormatName);
+			TextureFormat = TPM->FindTextureFormat(BuildSettings.TextureFormatName);
 		}
 
 		if (TextureFormat)
@@ -2225,10 +2226,10 @@ public:
 		)
 	{
 		const ITextureFormat* TextureFormat = nullptr;
-
-		if (TFM)
+		ITargetPlatformManagerModule* TPM = GetTargetPlatformManager();
+		if (TPM)
 		{
-			TextureFormat = TFM->FindTextureFormat(BuildSettings.TextureFormatName);
+			TextureFormat = TPM->FindTextureFormat(BuildSettings.TextureFormatName);
 		}
 		if (TextureFormat == nullptr)
 		{
@@ -2268,12 +2269,12 @@ public:
 
 			if(!BuildTextureMips(AssociatedNormalSourceMips, DefaultSettings, CompressorCaps, IntermediateAssociatedNormalSourceMipChain))
 			{
-				UE_LOG(LogTextureCompressor, Warning, TEXT("Failed to generate texture mips for composite texture"));
+				UE_LOG(LogTexture, Warning, TEXT("Failed to generate texture mips for composite texture"));
 			}
 
 			if(!ApplyCompositeTexture(IntermediateMipChain, IntermediateAssociatedNormalSourceMipChain, BuildSettings.CompositeTextureMode, BuildSettings.CompositePower))
 			{
-				UE_LOG(LogTextureCompressor, Warning, TEXT("Failed to apply composite texture"));
+				UE_LOG(LogTexture, Warning, TEXT("Failed to apply composite texture"));
 			}
 		}
 
@@ -2304,9 +2305,6 @@ public:
 		nvTextureToolsHandle = FPlatformProcess::GetDllHandle(*(FPaths::EngineDir() / TEXT("Binaries/ThirdParty/nvTextureTools/Win32/nvtt_.dll")));
 	#endif
 #endif	//PLATFORM_WINDOWS
-
-		// Initialize TFM singleton on the main thread, as it may be accessed from task threads later.
-		TFM = GetTextureFormatManager();
 	}
 
 	void ShutdownModule()
@@ -2323,8 +2321,6 @@ private:
 	void* nvTextureToolsHandle;
 #endif	//PLATFORM_WINDOWS
 
-	ITextureFormatManagerModule* TFM = nullptr;
-
 	bool BuildTextureMips(
 		const TArray<FImage>& InSourceMips,
 		const FTextureBuildSettings& BuildSettings,
diff --git a/Engine/Source/Developer/TextureCompressor/TextureCompressor.Build.cs b/Engine/Source/Developer/TextureCompressor/TextureCompressor.Build.cs
--- a/Engine/Source/Developer/TextureCompressor/TextureCompressor.Build.cs
+++ b/Engine/Source/Developer/TextureCompressor/TextureCompressor.Build.cs
@@ -6,17 +6,13 @@ public class TextureCompressor : ModuleRules
 {
 	public TextureCompressor(ReadOnlyTargetRules Target) : base(Target)
 	{
-		PrivateIncludePathModuleNames.AddRange(new string[]
-		{
-			"CoreUObject",
-			"Engine",
-		});
-
 		PrivateDependencyModuleNames.AddRange(
 			new string[] {
 				"Core",
-				"ImageCore",
-				"TextureFormat",
+				"CoreUObject", // @todo Mac: for some reason it's needed to link in debug on Mac
+				"Engine",
+				"TargetPlatform",
+				"ImageCore"
 			}
 			);
 
diff --git a/Engine/Source/Developer/TextureFormat/Private/TextureFormatManagerModule.cpp b/Engine/Source/Developer/TextureFormat/Private/TextureFormatManagerModule.cpp
deleted file mode 100644
diff --git a/Engine/Source/Developer/TextureFormat/Public/ChildTextureFormat.h b/Engine/Source/Developer/TextureFormat/Public/ChildTextureFormat.h
deleted file mode 100644
diff --git a/Engine/Source/Developer/TextureFormat/Public/Interfaces/ITextureFormat.h b/Engine/Source/Developer/TextureFormat/Public/Interfaces/ITextureFormat.h
deleted file mode 100644
diff --git a/Engine/Source/Developer/TextureFormat/Public/Interfaces/ITextureFormatManagerModule.h b/Engine/Source/Developer/TextureFormat/Public/Interfaces/ITextureFormatManagerModule.h
deleted file mode 100644
diff --git a/Engine/Source/Developer/TextureFormat/Public/Interfaces/ITextureFormatModule.h b/Engine/Source/Developer/TextureFormat/Public/Interfaces/ITextureFormatModule.h
deleted file mode 100644
diff --git a/Engine/Source/Developer/TextureFormat/Public/TextureFormatManager.h b/Engine/Source/Developer/TextureFormat/Public/TextureFormatManager.h
deleted file mode 100644
diff --git a/Engine/Source/Developer/TextureFormat/TextureFormat.Build.cs b/Engine/Source/Developer/TextureFormat/TextureFormat.Build.cs
deleted file mode 100644
diff --git a/Engine/Source/Runtime/Engine/Classes/Engine/Texture.h b/Engine/Source/Runtime/Engine/Classes/Engine/Texture.h
--- a/Engine/Source/Runtime/Engine/Classes/Engine/Texture.h
+++ b/Engine/Source/Runtime/Engine/Classes/Engine/Texture.h
@@ -72,6 +72,25 @@ enum TextureAddress
 	TA_MAX,
 };
 
+UENUM()
+enum ECompositeTextureMode
+{
+	CTM_Disabled UMETA(DisplayName="Disabled"),
+	/** CompositingTexture needs to be a normal map with the same or larger size. */
+	CTM_NormalRoughnessToRed UMETA(DisplayName="Add Normal Roughness To Red"),
+	/** CompositingTexture needs to be a normal map with the same or larger size. */
+	CTM_NormalRoughnessToGreen UMETA(DisplayName="Add Normal Roughness To Green"),
+	/** CompositingTexture needs to be a normal map with the same or larger size. */
+	CTM_NormalRoughnessToBlue UMETA(DisplayName="Add Normal Roughness To Blue"),
+	/** CompositingTexture needs to be a normal map with the same or larger size. */
+	CTM_NormalRoughnessToAlpha UMETA(DisplayName="Add Normal Roughness To Alpha"),
+	CTM_MAX,
+
+	// Note: These are serialized as as raw values in the texture DDC key, so additional entries
+	// should be added at the bottom; reordering or removing entries will require changing the GUID
+	// in the texture compressor DDC key
+};
+
 UENUM()
 enum ETextureMipCount
 {
diff --git a/Engine/Source/Runtime/Engine/Classes/Engine/TextureDefines.h b/Engine/Source/Runtime/Engine/Classes/Engine/TextureDefines.h
--- a/Engine/Source/Runtime/Engine/Classes/Engine/TextureDefines.h
+++ b/Engine/Source/Runtime/Engine/Classes/Engine/TextureDefines.h
@@ -4,14 +4,7 @@
 
 #include "CoreMinimal.h"
 #include "UObject/ObjectMacros.h"
-
-// TextureDefines.h is used from the TextureCompressor module as an "include only"
-// dependency to get visibility of these enum values without linking to the engine module
-// to facilitate that the generated TextureDefines header is only conditionally included
-// if we're compiling with the engine.
-#if WITH_ENGINE
 #include "TextureDefines.generated.h"
-#endif // WITH_ENGINE
 
 /**
  * @warning: if this is changed:
@@ -223,23 +216,3 @@ enum ETextureLossyCompressionAmount
 	TLCA_High			UMETA(DisplayName = "High"),
 	TLCA_Highest		UMETA(DisplayName = "Highest (Worst image quality, smallest filesize)"),
 };
-
-
-UENUM()
-enum ECompositeTextureMode
-{
-	CTM_Disabled UMETA(DisplayName="Disabled"),
-	/** CompositingTexture needs to be a normal map with the same or larger size. */
-	CTM_NormalRoughnessToRed UMETA(DisplayName="Add Normal Roughness To Red"),
-	/** CompositingTexture needs to be a normal map with the same or larger size. */
-	CTM_NormalRoughnessToGreen UMETA(DisplayName="Add Normal Roughness To Green"),
-	/** CompositingTexture needs to be a normal map with the same or larger size. */
-	CTM_NormalRoughnessToBlue UMETA(DisplayName="Add Normal Roughness To Blue"),
-	/** CompositingTexture needs to be a normal map with the same or larger size. */
-	CTM_NormalRoughnessToAlpha UMETA(DisplayName="Add Normal Roughness To Alpha"),
-	CTM_MAX,
-
-	// Note: These are serialized as as raw values in the texture DDC key, so additional entries
-	// should be added at the bottom; reordering or removing entries will require changing the GUID
-	// in the texture compressor DDC key
-};
