Adding missing ADPCM Audio format Module references
(Fixing build issue: "Unable to find audio format 'ADPCM'")

diff --git a/Engine/Source/Developer/TargetPlatform/TargetPlatform.Build.cs b/Engine/Source/Developer/TargetPlatform/TargetPlatform.Build.cs
--- a/Engine/Source/Developer/TargetPlatform/TargetPlatform.Build.cs
+++ b/Engine/Source/Developer/TargetPlatform/TargetPlatform.Build.cs
@@ -50,6 +50,7 @@ public class TargetPlatform : ModuleRules
 
 				if (Target.bCompileAgainstEngine)
 				{
+					DynamicallyLoadedModuleNames.Add("AudioFormatADPCM");
 					DynamicallyLoadedModuleNames.Add("AudioFormatOgg");
 					DynamicallyLoadedModuleNames.Add("AudioFormatOpus");
 				}
@@ -69,6 +70,7 @@ public class TargetPlatform : ModuleRules
 
 				if (Target.bCompileAgainstEngine)
 				{
+					DynamicallyLoadedModuleNames.Add("AudioFormatADPCM");
 					DynamicallyLoadedModuleNames.Add("AudioFormatOgg");
 					DynamicallyLoadedModuleNames.Add("AudioFormatOpus");
 				}
@@ -82,6 +84,7 @@ public class TargetPlatform : ModuleRules
 
 				if (Target.bCompileAgainstEngine)
 				{
+					DynamicallyLoadedModuleNames.Add("AudioFormatADPCM");
 					DynamicallyLoadedModuleNames.Add("AudioFormatOgg");
 					DynamicallyLoadedModuleNames.Add("AudioFormatOpus");
 				}
