- Fixed up some platform module references in C# for building with some platforms force disabled (brought over from UE4)

diff --git a/Engine/Source/Developer/TargetPlatform/TargetPlatform.Build.cs b/Engine/Source/Developer/TargetPlatform/TargetPlatform.Build.cs
--- a/Engine/Source/Developer/TargetPlatform/TargetPlatform.Build.cs
+++ b/Engine/Source/Developer/TargetPlatform/TargetPlatform.Build.cs
@@ -43,6 +43,7 @@ public class TargetPlatform : ModuleRules
             DynamicallyLoadedModuleNames.Add("ShaderFormatVectorVM");
             DynamicallyLoadedModuleNames.Add("ImageWrapper");
 
+			// @todo: should move all of this to specific UEBuild*.cs files
 			if (Target.Platform == UnrealTargetPlatform.Win64)
 			{
 				DynamicallyLoadedModuleNames.Add("TextureFormatIntelISPCTexComp");
@@ -60,14 +61,6 @@ public class TargetPlatform : ModuleRules
 					DynamicallyLoadedModuleNames.Add("AudioFormatOpus");
 					DynamicallyLoadedModuleNames.Add("AudioFormatBink");
 				}
-
-				if (Target.Type == TargetType.Editor || Target.Type == TargetType.Program)
-				{
-					DynamicallyLoadedModuleNames.Add("AndroidTargetPlatform");
-					DynamicallyLoadedModuleNames.Add("IOSTargetPlatform");
-					DynamicallyLoadedModuleNames.Add("TVOSTargetPlatform");
-					DynamicallyLoadedModuleNames.Add("MacTargetPlatform");
-				}
 			}
 			else if (Target.Platform == UnrealTargetPlatform.Mac)
 			{
diff --git a/Engine/Source/Editor/UnrealEd/UnrealEd.Build.cs b/Engine/Source/Editor/UnrealEd/UnrealEd.Build.cs
--- a/Engine/Source/Editor/UnrealEd/UnrealEd.Build.cs
+++ b/Engine/Source/Editor/UnrealEd/UnrealEd.Build.cs
@@ -282,21 +282,6 @@ public class UnrealEd : ModuleRules
 			);
 		}
 
-		if (Target.Platform == UnrealTargetPlatform.Win64 || Target.Platform == UnrealTargetPlatform.Mac)
-		{
-			DynamicallyLoadedModuleNames.Add("IOSPlatformEditor");
-		}
-
-		if (Target.Platform == UnrealTargetPlatform.Win64 || Target.Platform == UnrealTargetPlatform.Mac || Target.Platform == UnrealTargetPlatform.Linux)
-		{
-			DynamicallyLoadedModuleNames.Add("AndroidPlatformEditor");
-		}
-
-		if (Target.Platform == UnrealTargetPlatform.Win64 || Target.Platform == UnrealTargetPlatform.Mac)
-		{
-			DynamicallyLoadedModuleNames.Add("LuminPlatformEditor");
-		}
-
 		CircularlyReferencedDependentModules.AddRange(
 			new string[] {
 				"Documentation",
diff --git a/Engine/Source/Programs/UnrealBuildTool/Platform/Android/UEBuildAndroid.cs b/Engine/Source/Programs/UnrealBuildTool/Platform/Android/UEBuildAndroid.cs
--- a/Engine/Source/Programs/UnrealBuildTool/Platform/Android/UEBuildAndroid.cs
+++ b/Engine/Source/Programs/UnrealBuildTool/Platform/Android/UEBuildAndroid.cs
@@ -317,6 +317,11 @@ namespace UnrealBuildTool
 						//Rules.DynamicallyLoadedModuleNames.Add("ShaderFormatAndroid");		//@todo android: ShaderFormatAndroid
 					}
 				}
+
+				if (ModuleName == "UnrealEd")
+				{
+					Rules.DynamicallyLoadedModuleNames.Add("AndroidPlatformEditor");
+				}
 			}
 		}
 
diff --git a/Engine/Source/Programs/UnrealBuildTool/Platform/Lumin/UEBuildLumin.cs b/Engine/Source/Programs/UnrealBuildTool/Platform/Lumin/UEBuildLumin.cs
--- a/Engine/Source/Programs/UnrealBuildTool/Platform/Lumin/UEBuildLumin.cs
+++ b/Engine/Source/Programs/UnrealBuildTool/Platform/Lumin/UEBuildLumin.cs
@@ -132,6 +132,11 @@ namespace UnrealBuildTool
 						Rules.DynamicallyLoadedModuleNames.Add("LuminTargetPlatform");
 					}
 				}
+
+				if (ModuleName == "UnrealEd")
+				{
+					Rules.DynamicallyLoadedModuleNames.Add("LuminPlatformEditor");
+				}
 			}
 		}
 
