No unity fix
#rnx

diff --git a/Engine/Source/Developer/TargetPlatform/Private/TargetPlatformManagerModule.cpp b/Engine/Source/Developer/TargetPlatform/Private/TargetPlatformManagerModule.cpp
--- a/Engine/Source/Developer/TargetPlatform/Private/TargetPlatformManagerModule.cpp
+++ b/Engine/Source/Developer/TargetPlatform/Private/TargetPlatformManagerModule.cpp
@@ -2,6 +2,7 @@
 
 #include "CoreMinimal.h"
 #include "HAL/FileManager.h"
+#include "Misc/ConfigCacheIni.h"
 #include "Misc/CommandLine.h"
 #include "Misc/Paths.h"
 #include "Misc/OutputDeviceRedirector.h"
