We now have 2 or 3 lines potentially in CurrentlyInstalled.txt AutoSDK file. Handle both cases

#[fyi] Josh.Adams

diff --git a/Engine/Source/Developer/TargetPlatform/Private/TargetPlatformManagerModule.cpp b/Engine/Source/Developer/TargetPlatform/Private/TargetPlatformManagerModule.cpp
--- a/Engine/Source/Developer/TargetPlatform/Private/TargetPlatformManagerModule.cpp
+++ b/Engine/Source/Developer/TargetPlatform/Private/TargetPlatformManagerModule.cpp
@@ -949,7 +949,7 @@ protected:
 			InstallManifestFile->Close();
 
 			
-			if (FileLines.Num() != 2)
+			if (FileLines.Num() != 2 && FileLines.Num() != 3)
 			{
 				UE_LOG(LogTargetPlatformManager, Warning, TEXT("Malformed install manifest file for Platform %s"), *AutoSDKPath);
 				return false;
