Expose per-project platform setting to control offline BVHs build mode (Disabled, MaximizePerformance, MinimizeMemory).
- Offline BVHs generally result in better performance and/or reduced memory usage.
- At the moment MaximizePerformance vs MinimizeMemory control whether triangle splits are used.

diff --git a/Engine/Source/Developer/TargetPlatform/Private/TargetPlatformBase.cpp b/Engine/Source/Developer/TargetPlatform/Private/TargetPlatformBase.cpp
--- a/Engine/Source/Developer/TargetPlatform/Private/TargetPlatformBase.cpp
+++ b/Engine/Source/Developer/TargetPlatform/Private/TargetPlatformBase.cpp
@@ -54,6 +54,11 @@ bool FTargetPlatformBase::UsesRayTracing() const
 	return CVar ? (CVar->GetInt() != 0) : false;
 }
 
+EOfflineBVHMode FTargetPlatformBase::GetStaticMeshOfflineBVHMode() const
+{
+	return EOfflineBVHMode::Disabled;
+}
+
 bool FTargetPlatformBase::ForcesSimpleSkyDiffuse() const
 {
 	return false;
diff --git a/Engine/Source/Developer/TargetPlatform/Public/Common/TargetPlatformBase.h b/Engine/Source/Developer/TargetPlatform/Public/Common/TargetPlatformBase.h
--- a/Engine/Source/Developer/TargetPlatform/Public/Common/TargetPlatformBase.h
+++ b/Engine/Source/Developer/TargetPlatform/Public/Common/TargetPlatformBase.h
@@ -67,6 +67,8 @@ public:
 
 	TARGETPLATFORM_API virtual bool UsesRayTracing() const override;
 
+	TARGETPLATFORM_API virtual EOfflineBVHMode GetStaticMeshOfflineBVHMode() const override;
+
 	TARGETPLATFORM_API virtual bool ForcesSimpleSkyDiffuse() const override;
 
 	TARGETPLATFORM_API virtual float GetDownSampleMeshDistanceFieldDivider() const override;
diff --git a/Engine/Source/Developer/TargetPlatform/Public/Interfaces/ITargetPlatform.h b/Engine/Source/Developer/TargetPlatform/Public/Interfaces/ITargetPlatform.h
--- a/Engine/Source/Developer/TargetPlatform/Public/Interfaces/ITargetPlatform.h
+++ b/Engine/Source/Developer/TargetPlatform/Public/Interfaces/ITargetPlatform.h
@@ -116,6 +116,13 @@ enum class EPlatformAuthentication
 	Always,
 };
 
+enum class EOfflineBVHMode
+{
+	Disabled,
+	MaximizePerformance,
+	MinimizeMemory,
+};
+
 /**
  * Flags specifying what is needed to be able to complete and deploy a build.
  */
@@ -466,6 +473,11 @@ public:
 	*/
 	virtual bool UsesRayTracing() const = 0;
 
+	/**
+	* Gets static mesh offline BVH mode
+	*/
+	virtual EOfflineBVHMode GetStaticMeshOfflineBVHMode() const = 0;
+
 	/**
 	* Gets whether the platform will use SH2 instead of SH3 for sky irradiance.
 	*/
