- Kick off UBT to setup AutoSDK early, and async, so editor doesn't need to wait for it to run
- Deprecated unuesed GetCookingTargetPlatforms()
#p4v-preflight-copy 18811010

diff --git a/Engine/Source/Developer/TargetPlatform/Private/TargetPlatformManagerModule.cpp b/Engine/Source/Developer/TargetPlatform/Private/TargetPlatformManagerModule.cpp
--- a/Engine/Source/Developer/TargetPlatform/Private/TargetPlatformManagerModule.cpp
+++ b/Engine/Source/Developer/TargetPlatform/Private/TargetPlatformManagerModule.cpp
@@ -37,6 +37,35 @@ DEFINE_LOG_CATEGORY_STATIC(LogTargetPlatformManager, Log, All);
 #define AUTOSDKS_ENABLED DDPI_HAS_EXTENDED_PLATFORMINFO_DATA
 #endif
 
+#if AUTOSDKS_ENABLED
+namespace UE::AutoSDK
+{
+
+// kick off a call to UBT nice and early so that it's results are hopefully ready when needed
+static FProcHandle AutoSDKSetupUBTProc;
+FDelayedAutoRegisterHelper GAutoSDKInit(EDelayedRegisterRunPhase::FileSystemReady, []
+	{
+		// amortize UBT cost by calling it once for all platforms, rather than once per platform.
+		if (FParse::Param(FCommandLine::Get(), TEXT("Multiprocess")) == false)
+		{
+			FString UBTParams(TEXT("-Mode=SetupPlatforms"));
+			int32 UBTReturnCode = -1;
+			FString UBTOutput;
+
+			void* ReadPipe;
+			void* WritePipe;
+			AutoSDKSetupUBTProc = FDesktopPlatformModule::Get()->InvokeUnrealBuildToolAsync(UBTParams, *GLog, ReadPipe, WritePipe, true);
+			if (!AutoSDKSetupUBTProc.IsValid())
+			{
+				UE_LOG(LogTargetPlatformManager, Warning, TEXT("AutoSDK is enabled (UE_SDKS_ROOT is set), but failed to run UBT to check SDK status! Check your installation."));
+			}
+		}
+	}
+);
+
+}
+#endif
+
 
 static const size_t MaxPlatformCount = 64;		// In the unlikely event that someone bumps this please note that there's
 												// an implicit assumption that there won't be more than 64 unique target
@@ -95,6 +124,7 @@ public:
 		, bIgnoreFirstDelegateCall(true)
 	{
 #if WITH_EDITOR
+
 		ITurnkeySupportModule::Get().UpdateSdkInfo();
 #endif
 
@@ -103,22 +133,15 @@ public:
 		// AutoSDKs only enabled if UE_SDKS_ROOT is set.
 		if (IsAutoSDKsEnabled())
 		{					
-			DECLARE_SCOPE_CYCLE_COUNTER( TEXT( "FTargetPlatformManagerModule.StartAutoSDK" ), STAT_FTargetPlatformManagerModule_StartAutoSDK, STATGROUP_TargetPlatform );
-
-			// amortize UBT cost by calling it once for all platforms, rather than once per platform.
-			if (FParse::Param(FCommandLine::Get(), TEXT("Multiprocess"))==false)
+			if (UE::AutoSDK::AutoSDKSetupUBTProc.IsValid())
 			{
-				FString UBTParams(TEXT("-Mode=SetupPlatforms"));
-				int32 UBTReturnCode = -1;
-				FString UBTOutput;
-				if (!FDesktopPlatformModule::Get()->InvokeUnrealBuildToolSync(UBTParams, *GLog, true, UBTReturnCode, UBTOutput))
-				{
-					UE_LOG(LogTargetPlatformManager, Warning, TEXT("AutoSDK is enabled (UE_SDKS_ROOT is set), but failed to run UBT to check SDK status! Check your installation."));
-				}
+				SCOPED_BOOT_TIMING("FTargetPlatformManagerModule - WaitForUBTProc");
+				FPlatformProcess::WaitForProc(UE::AutoSDK::AutoSDKSetupUBTProc);
 			}
 
 			// we have to setup our local environment according to AutoSDKs or the ITargetPlatform's IsSDkInstalled calls may fail
 			// before we get a change to setup for a given platform.  Use the platforminfo list to avoid any kind of interdependency.
+			SCOPED_BOOT_TIMING("FTargetPlatformManagerModule.SetupAndValidateAutoSDK");
 			for (auto Pair: FDataDrivenPlatformInfoRegistry::GetAllPlatformInfos())
 			{
 				if (Pair.Value.AutoSDKPath.Len() > 0)
@@ -250,28 +273,7 @@ public:
 
 	virtual const TArray<ITargetPlatform*>& GetCookingTargetPlatforms() override
 	{
-		static bool bInitialized = false;
-		static TArray<ITargetPlatform*> Results;
-
-		if ( !bInitialized || bForceCacheUpdate )
-		{
-			Results = GetActiveTargetPlatforms();
-
-			FString PlatformStr;
-			if (FParse::Value(FCommandLine::Get(), TEXT("TARGETPLATFORM="), PlatformStr))
-			{
-				if (PlatformStr == TEXT("None"))
-				{
-					Results = Platforms;
-				}
-			}
-			else
-			{
-				Results = Platforms;
-			}
-		}
-
-		return Results;
+		return GetActiveTargetPlatforms();
 	}
 
 	virtual const TArray<ITargetPlatform*>& GetActiveTargetPlatforms() override
diff --git a/Engine/Source/Developer/TargetPlatform/Public/Interfaces/ITargetPlatformManagerModule.h b/Engine/Source/Developer/TargetPlatform/Public/Interfaces/ITargetPlatformManagerModule.h
--- a/Engine/Source/Developer/TargetPlatform/Public/Interfaces/ITargetPlatformManagerModule.h
+++ b/Engine/Source/Developer/TargetPlatform/Public/Interfaces/ITargetPlatformManagerModule.h
@@ -98,6 +98,7 @@ public:
 	 *
 	 * @return Collection of platforms.
 	 */
+	UE_DEPRECATED(5.1, "Use GetActiveTargetPlatforms instead of GetCookingTargetPlatforms")
 	virtual const TArray<ITargetPlatform*>& GetCookingTargetPlatforms() = 0;
 
 	/**
