Fix crash dereferencing garbage WritePtr in Linux

Running with:

   Engine/Binaries/Linux/UnrealEditor-Linux-Debug -logcmds="LogAnalytics VeryVerbose" -AnalyticsDisableCaching

Crashed on GetHandle() of invalid write pointer. Code introduced in CL 18846806.

diff --git a/Engine/Source/Developer/TargetPlatform/Private/TargetPlatformManagerModule.cpp b/Engine/Source/Developer/TargetPlatform/Private/TargetPlatformManagerModule.cpp
--- a/Engine/Source/Developer/TargetPlatform/Private/TargetPlatformManagerModule.cpp
+++ b/Engine/Source/Developer/TargetPlatform/Private/TargetPlatformManagerModule.cpp
@@ -52,8 +52,8 @@ FDelayedAutoRegisterHelper GAutoSDKInit(EDelayedRegisterRunPhase::FileSystemRead
 			int32 UBTReturnCode = -1;
 			FString UBTOutput;
 
-			void* ReadPipe;
-			void* WritePipe;
+			void* ReadPipe = nullptr;
+			void* WritePipe = nullptr;
 			AutoSDKSetupUBTProc = FDesktopPlatformModule::Get()->InvokeUnrealBuildToolAsync(UBTParams, *GLog, ReadPipe, WritePipe, true);
 			if (!AutoSDKSetupUBTProc.IsValid())
 			{
