- Moved a call of IsAutoSDKsEnabled() up into the async pre-kick of UBT during editor startup to fix warnings when UE_SDKS_ROOT isn't set
#p4v-preflight-copy 18902221

diff --git a/Engine/Source/Developer/TargetPlatform/Private/TargetPlatformManagerModule.cpp b/Engine/Source/Developer/TargetPlatform/Private/TargetPlatformManagerModule.cpp
--- a/Engine/Source/Developer/TargetPlatform/Private/TargetPlatformManagerModule.cpp
+++ b/Engine/Source/Developer/TargetPlatform/Private/TargetPlatformManagerModule.cpp
@@ -40,13 +40,26 @@ DEFINE_LOG_CATEGORY_STATIC(LogTargetPlatformManager, Log, All);
 #if AUTOSDKS_ENABLED
 namespace UE::AutoSDK
 {
+static bool IsAutoSDKsEnabled()
+{
+	static const FString SDKRootEnvFar(TEXT("UE_SDKS_ROOT"));
+
+	FString SDKPath = FPlatformMisc::GetEnvironmentVariable(*SDKRootEnvFar);
+
+	// AutoSDKs only enabled if UE_SDKS_ROOT is set.
+	if (SDKPath.Len() != 0)
+	{
+		return true;
+	}
+	return false;
+}
 
 // kick off a call to UBT nice and early so that it's results are hopefully ready when needed
 static FProcHandle AutoSDKSetupUBTProc;
 FDelayedAutoRegisterHelper GAutoSDKInit(EDelayedRegisterRunPhase::FileSystemReady, []
 	{
 		// amortize UBT cost by calling it once for all platforms, rather than once per platform.
-		if (FParse::Param(FCommandLine::Get(), TEXT("Multiprocess")) == false)
+		if (IsAutoSDKsEnabled() && FParse::Param(FCommandLine::Get(), TEXT("Multiprocess")) == false)
 		{
 			FString UBTParams(TEXT("-Mode=SetupPlatforms"));
 			int32 UBTReturnCode = -1;
@@ -131,7 +144,7 @@ public:
 #if AUTOSDKS_ENABLED		
 		
 		// AutoSDKs only enabled if UE_SDKS_ROOT is set.
-		if (IsAutoSDKsEnabled())
+		if (UE::AutoSDK::IsAutoSDKsEnabled())
 		{					
 			if (UE::AutoSDK::AutoSDKSetupUBTProc.IsValid())
 			{
@@ -735,24 +748,6 @@ public:
 
 protected:
 
-	/**
-	 * Checks whether AutoSDK is enabled.
-	 *
-	 * @return true if the SDK is enabled, false otherwise.
-	 */
-	bool IsAutoSDKsEnabled()
-	{
-		static const FString SDKRootEnvFar(TEXT("UE_SDKS_ROOT"));
-
-		FString SDKPath = FPlatformMisc::GetEnvironmentVariable(*SDKRootEnvFar);
-
-		// AutoSDKs only enabled if UE_SDKS_ROOT is set.
-		if (SDKPath.Len() != 0)
-		{
-			return true;
-		}
-		return false;
-	}
 
 	bool InitializeSinglePlatform(FName PlatformName, const FString& AutoSDKPath)
 	{
@@ -913,7 +908,7 @@ protected:
 	{						
 #if AUTOSDKS_ENABLED
 		
-		if (!IsAutoSDKsEnabled())
+		if (!UE::AutoSDK::IsAutoSDKsEnabled())
 		{
 			return true;
 		}
