FindTargetPlatformWithSupport gives priority to the currently active/cooking platforms for cases where multiple platforms support the shader format.

 - fixes ShaderPlatformQualitySettings overrides in platform engine ini for platforms that share the same shader format.
 - still a potential issue if trying to cook these platforms at the same time ( -targetplatform=Plat1+Plat2 )

#rnx

diff --git a/Engine/Source/Developer/TargetPlatform/Private/TargetPlatformManagerModule.cpp b/Engine/Source/Developer/TargetPlatform/Private/TargetPlatformManagerModule.cpp
--- a/Engine/Source/Developer/TargetPlatform/Private/TargetPlatformManagerModule.cpp
+++ b/Engine/Source/Developer/TargetPlatform/Private/TargetPlatformManagerModule.cpp
@@ -271,13 +271,18 @@ public:
 
 	virtual ITargetPlatform* FindTargetPlatformWithSupport(FName SupportType, FName RequiredSupportedValue)
 	{
-		const TArray<ITargetPlatform*>& TargetPlatforms = GetTargetPlatforms();
-
-		for (int32 Index = 0; Index < TargetPlatforms.Num(); Index++)
+		// first try to find an active target platform. if that fails, try all target platforms.
+		// this gives priority to the active target platform if multiple platforms support the same value
+		for (int Pass = 0; Pass < 2; Pass++)
 		{
-			if (TargetPlatforms[Index]->SupportsValueForType(SupportType, RequiredSupportedValue))
+			const TArray<ITargetPlatform*>& TargetPlatforms = (Pass == 0) ? GetActiveTargetPlatforms() : GetTargetPlatforms();
+
+			for (int32 Index = 0; Index < TargetPlatforms.Num(); Index++)
 			{
-				return TargetPlatforms[Index];
+				if (TargetPlatforms[Index]->SupportsValueForType(SupportType, RequiredSupportedValue))
+				{
+					return TargetPlatforms[Index];
+				}
 			}
 		}
 
