Add scope marker to see what platform we are trying to load.

- This can take a considerable amount of time so it is good to know which platforms are loading.

diff --git a/Engine/Source/Developer/TargetPlatform/Private/TargetPlatformManagerModule.cpp b/Engine/Source/Developer/TargetPlatform/Private/TargetPlatformManagerModule.cpp
--- a/Engine/Source/Developer/TargetPlatform/Private/TargetPlatformManagerModule.cpp
+++ b/Engine/Source/Developer/TargetPlatform/Private/TargetPlatformManagerModule.cpp
@@ -836,6 +836,9 @@ protected:
 		for (auto Pair : FDataDrivenPlatformInfoRegistry::GetAllPlatformInfos())
 		{
 			FName PlatformName = Pair.Key;
+
+			TRACE_CPUPROFILER_EVENT_SCOPE_TEXT(*PlatformName.ToString());
+
 			const FDataDrivenPlatformInfo& Info = Pair.Value;
 
 			SlowTask.EnterProgressFrame(1);
