remove redundent HoloLens platform check in disabled code.

#rnx

diff --git a/Engine/Source/Developer/TargetPlatform/Private/TargetPlatformManagerModule.cpp b/Engine/Source/Developer/TargetPlatform/Private/TargetPlatformManagerModule.cpp
--- a/Engine/Source/Developer/TargetPlatform/Private/TargetPlatformManagerModule.cpp
+++ b/Engine/Source/Developer/TargetPlatform/Private/TargetPlatformManagerModule.cpp
@@ -1176,11 +1176,6 @@ protected:
 					// since Desktop is just packaging, we don't need an SDK, and UBT will return INVALID, since it doesn't build for it
 					PlatformInfo::UpdatePlatformSDKStatus(PlatformName, PlatformInfo::EPlatformSDKStatus::Installed);
 				}
-				else if (PlatformName == TEXT("HoloLens"))
-				{
-					PlatformName = TEXT("HoloLens");
-					PlatformInfo::UpdatePlatformSDKStatus(PlatformName, Status);
-				}
 				else
 				{
 					PlatformInfo::UpdatePlatformSDKStatus(PlatformName, Status);
