Fix for nounity CIS in TargetPlatformBase.h
#rnx

diff --git a/Engine/Source/Developer/TargetPlatform/Public/Common/TargetPlatformBase.h b/Engine/Source/Developer/TargetPlatform/Public/Common/TargetPlatformBase.h
--- a/Engine/Source/Developer/TargetPlatform/Public/Common/TargetPlatformBase.h
+++ b/Engine/Source/Developer/TargetPlatform/Public/Common/TargetPlatformBase.h
@@ -10,6 +10,7 @@
 #include "Misc/ConfigCacheIni.h"
 #include "Misc/Paths.h"
 #include "PlatformInfo.h"
+#include "Templates/PimplPtr.h"
 
 // Forward declare.
 namespace Audio 
