Adding SupportsRayTracing to platform properties. This is checked in addition to the project/platform settings for RayTracing but is opt-in per platform instead of always enabled if the setting is enabled.
This has the effect of reducing the number of platforms that say they support RayTracing which directly reduces the number of shader platforms that "support" RayTracing.

diff --git a/Engine/Source/Developer/TargetPlatform/Public/Common/TargetPlatformBase.h b/Engine/Source/Developer/TargetPlatform/Public/Common/TargetPlatformBase.h
--- a/Engine/Source/Developer/TargetPlatform/Public/Common/TargetPlatformBase.h
+++ b/Engine/Source/Developer/TargetPlatform/Public/Common/TargetPlatformBase.h
@@ -497,6 +497,15 @@ public:
 		return TPlatformProperties::GetMemoryMappingAlignment();
 	}
 
+	virtual bool UsesRayTracing() const override
+	{
+		if (TPlatformProperties::SupportsRayTracing())
+		{
+			return FTargetPlatformBase::UsesRayTracing();
+		}
+		return false;
+	}
+
 
 #if WITH_ENGINE
 	virtual FName GetPhysicsFormat( class UBodySetup* Body ) const override
diff --git a/Engine/Source/Runtime/Core/Public/GenericPlatform/GenericPlatformProperties.h b/Engine/Source/Runtime/Core/Public/GenericPlatform/GenericPlatformProperties.h
--- a/Engine/Source/Runtime/Core/Public/GenericPlatform/GenericPlatformProperties.h
+++ b/Engine/Source/Runtime/Core/Public/GenericPlatform/GenericPlatformProperties.h
@@ -256,6 +256,11 @@ struct FGenericPlatformProperties
 		return false; // Currently VT is opt-in
 	}
 
+	static FORCEINLINE bool SupportsRayTracing()
+	{
+		return false;
+	}
+
 	static FORCEINLINE bool SupportsLumenGI()
 	{
 		return true;
diff --git a/Engine/Source/Runtime/Core/Public/Linux/LinuxPlatformProperties.h b/Engine/Source/Runtime/Core/Public/Linux/LinuxPlatformProperties.h
--- a/Engine/Source/Runtime/Core/Public/Linux/LinuxPlatformProperties.h
+++ b/Engine/Source/Runtime/Core/Public/Linux/LinuxPlatformProperties.h
@@ -120,6 +120,11 @@ struct FLinuxPlatformProperties
 		return true;
 	}
 
+	static FORCEINLINE bool SupportsRayTracing()
+	{
+		return true;
+	}
+
 	static FORCEINLINE bool SupportsQuit()
 	{
 		return true;
diff --git a/Engine/Source/Runtime/Core/Public/Windows/WindowsPlatformProperties.h b/Engine/Source/Runtime/Core/Public/Windows/WindowsPlatformProperties.h
--- a/Engine/Source/Runtime/Core/Public/Windows/WindowsPlatformProperties.h
+++ b/Engine/Source/Runtime/Core/Public/Windows/WindowsPlatformProperties.h
@@ -93,6 +93,11 @@ struct FWindowsPlatformProperties
 		return !IsServerOnly() && !HasEditorOnlyData();
 	}
 
+	static FORCEINLINE bool SupportsRayTracing()
+	{
+		return true;
+	}
+
 	static FORCEINLINE bool SupportsGrayscaleSRGB()
 	{
 		return false; // Requires expand from G8 to RGBA
