[TargetPlatform] minor cleanup GetActiveTargetPlatforms to only build error message when there is an error

#p4v-preflight-copy 21725742
#p4v-preflight-copy 21741782

diff --git a/Engine/Source/Developer/TargetPlatform/Private/TargetPlatformManagerModule.cpp b/Engine/Source/Developer/TargetPlatform/Private/TargetPlatformManagerModule.cpp
--- a/Engine/Source/Developer/TargetPlatform/Private/TargetPlatformManagerModule.cpp
+++ b/Engine/Source/Developer/TargetPlatform/Private/TargetPlatformManagerModule.cpp
@@ -334,27 +334,27 @@ public:
 
 					PlatformStr.ParseIntoArray(PlatformNames, TEXT("+"), true);
 
-					// for nicer user response
-					FString AvailablePlatforms;
-
 					for (int32 Index = 0; Index < TargetPlatforms.Num(); Index++)
 					{
 						if (PlatformNames.Contains(TargetPlatforms[Index]->PlatformName()))
-						{							
-							Results.Add(TargetPlatforms[Index]);						
-						}
-
-						if(!AvailablePlatforms.IsEmpty())
 						{
-							AvailablePlatforms += TEXT(", ");
+							Results.Add(TargetPlatforms[Index]);
 						}
-						AvailablePlatforms += TargetPlatforms[Index]->PlatformName();
 					}
 
 					if (Results.Num() == 0)
 					{
 						// An invalid platform was specified...
 						// Inform the user.
+						TStringBuilder<1024> AvailablePlatforms;
+						for (int32 Index = 0; Index < TargetPlatforms.Num(); Index++)
+						{
+							if (Index > 0)
+							{
+								AvailablePlatforms << TEXT(", ");
+							}
+							AvailablePlatforms << TargetPlatforms[Index]->PlatformName();
+						}
 						bHasInitErrors = true;
 						InitErrorMessages.Appendf(TEXT("Invalid target platform specified (%s). Available = { %s } "), *PlatformStr, *AvailablePlatforms);
 						UE_LOG(LogTargetPlatformManager, Error, TEXT("Invalid target platform specified (%s). Available = { %s } "), *PlatformStr, *AvailablePlatforms);
@@ -370,8 +370,8 @@ public:
 				for (int32 Index = 0; Index < TargetPlatforms.Num(); Index++)
 				{
 					if (TargetPlatforms[Index]->IsRunningPlatform())
-					{						
-						Results.Add(TargetPlatforms[Index]);					
+					{
+						Results.Add(TargetPlatforms[Index]);
 					}
 				}
 			}
