Capture platform-specific project analytics for launching, cooking & packaging in the editor.

diff --git a/Engine/Platforms/Hololens/Source/Developer/HoloLensTargetPlatform/Private/HoloLensTargetPlatform.cpp b/Engine/Platforms/Hololens/Source/Developer/HoloLensTargetPlatform/Private/HoloLensTargetPlatform.cpp
--- a/Engine/Platforms/Hololens/Source/Developer/HoloLensTargetPlatform/Private/HoloLensTargetPlatform.cpp
+++ b/Engine/Platforms/Hololens/Source/Developer/HoloLensTargetPlatform/Private/HoloLensTargetPlatform.cpp
@@ -180,6 +180,15 @@ void FHoloLensTargetPlatform::GetAllTargetedShaderFormats(TArray<FName>& OutForm
 	OutFormats.AddUnique(NAME_PCD3D_SM5);
 }
 
+void FHoloLensTargetPlatform::GetPlatformSpecificProjectAnalytics( TArray<FAnalyticsEventAttribute>& AnalyticsParamArray ) const
+{
+	TNonDesktopTargetPlatformBase<FHoloLensPlatformProperties>::GetPlatformSpecificProjectAnalytics(AnalyticsParamArray);
+
+	// NB. HololensDynamicRHI reads the Windows key for the default RHI
+	AppendAnalyticsEventConfigString(AnalyticsParamArray, TEXT("/Script/WindowsTargetPlatform.WindowsTargetSettings"), TEXT("DefaultGraphicsRHI"), GEngineIni);
+}
+
+
 #if WITH_ENGINE
 
 void FHoloLensTargetPlatform::GetReflectionCaptureFormats(TArray<FName>& OutFormats) const
diff --git a/Engine/Platforms/Hololens/Source/Developer/HoloLensTargetPlatform/Private/HoloLensTargetPlatform.h b/Engine/Platforms/Hololens/Source/Developer/HoloLensTargetPlatform/Private/HoloLensTargetPlatform.h
--- a/Engine/Platforms/Hololens/Source/Developer/HoloLensTargetPlatform/Private/HoloLensTargetPlatform.h
+++ b/Engine/Platforms/Hololens/Source/Developer/HoloLensTargetPlatform/Private/HoloLensTargetPlatform.h
@@ -67,6 +67,8 @@ public:
 
 	virtual void GetAllTargetedShaderFormats(TArray<FName>& OutFormats) const override;	
 
+	virtual void GetPlatformSpecificProjectAnalytics( TArray<struct FAnalyticsEventAttribute>& AnalyticsParamArray ) const override;
+
 #if WITH_ENGINE
 	virtual void GetReflectionCaptureFormats(TArray<FName>& OutFormats) const override;
 
diff --git a/Engine/Plugins/TraceUtilities/Source/EditorTraceUtilities/Private/UnrealInsightsLauncher.cpp b/Engine/Plugins/TraceUtilities/Source/EditorTraceUtilities/Private/UnrealInsightsLauncher.cpp
--- a/Engine/Plugins/TraceUtilities/Source/EditorTraceUtilities/Private/UnrealInsightsLauncher.cpp
+++ b/Engine/Plugins/TraceUtilities/Source/EditorTraceUtilities/Private/UnrealInsightsLauncher.cpp
@@ -120,7 +120,7 @@ void FUnrealInsightsLauncher::TryBuildUnrealInsightsExe(const FString& Path, con
 #endif
 
 	IUATHelperModule::Get().CreateUatTask(Arguments, PlatformName, LOCTEXT("BuildingUnrealInsights", "Building Unreal Insights"),
-		LOCTEXT("BuildUnrealInsightsTask", "Build Unreal Insights Task"), FAppStyle::GetBrush(TEXT("MainFrame.CookContent")), [this, Path, LaunchParameters](FString Result, double Time)
+		LOCTEXT("BuildUnrealInsightsTask", "Build Unreal Insights Task"), FAppStyle::GetBrush(TEXT("MainFrame.CookContent")), nullptr, [this, Path, LaunchParameters](FString Result, double Time)
 		{
 			if (Result.Equals(TEXT("Completed")))
 			{
diff --git a/Engine/Source/Developer/Android/AndroidTargetPlatform/Private/AndroidTargetPlatform.cpp b/Engine/Source/Developer/Android/AndroidTargetPlatform/Private/AndroidTargetPlatform.cpp
--- a/Engine/Source/Developer/Android/AndroidTargetPlatform/Private/AndroidTargetPlatform.cpp
+++ b/Engine/Source/Developer/Android/AndroidTargetPlatform/Private/AndroidTargetPlatform.cpp
@@ -26,6 +26,7 @@
 #include "Interfaces/IAndroidDeviceDetection.h"
 #include "Modules/ModuleManager.h"
 #include "Misc/SecureHash.h"
+#include "AnalyticsEventAttribute.h"
 
 
 #if WITH_ENGINE
@@ -439,6 +440,19 @@ void FAndroidTargetPlatform::GetAllTargetedShaderFormats( TArray<FName>& OutForm
 	GetAllPossibleShaderFormats(OutFormats);
 }
 
+void FAndroidTargetPlatform::GetPlatformSpecificProjectAnalytics( TArray<FAnalyticsEventAttribute>& AnalyticsParamArray ) const
+{
+	TNonDesktopTargetPlatformBase<FAndroidPlatformProperties>::GetPlatformSpecificProjectAnalytics( AnalyticsParamArray );
+
+	AppendAnalyticsEventAttributeArray(AnalyticsParamArray,
+		TEXT("AndroidVariant"), GetAndroidVariantName(),
+		TEXT("SupportsVulkan"), SupportsVulkan(),
+		TEXT("SupportsVulkanSM5"), SupportsVulkanSM5(),
+		TEXT("SupportsES31"), SupportsES31()
+	);
+
+	AppendAnalyticsEventConfigArray(AnalyticsParamArray, TEXT("/Script/AndroidRuntimeSettings.AndroidRuntimeSettings"), TEXT("PackageForOculusMobile"), GEngineIni);
+}
 
 #if WITH_ENGINE
 
diff --git a/Engine/Source/Developer/Android/AndroidTargetPlatform/Private/AndroidTargetPlatform.h b/Engine/Source/Developer/Android/AndroidTargetPlatform/Private/AndroidTargetPlatform.h
--- a/Engine/Source/Developer/Android/AndroidTargetPlatform/Private/AndroidTargetPlatform.h
+++ b/Engine/Source/Developer/Android/AndroidTargetPlatform/Private/AndroidTargetPlatform.h
@@ -122,6 +122,8 @@ public:
 	virtual void GetAllPossibleShaderFormats( TArray<FName>& OutFormats ) const override;
 
 	virtual void GetAllTargetedShaderFormats(TArray<FName>& OutFormats) const override;
+
+	virtual void GetPlatformSpecificProjectAnalytics( TArray<struct FAnalyticsEventAttribute>& AnalyticsParamArray ) const override;
 	
 #if WITH_ENGINE
 	virtual void GetReflectionCaptureFormats( TArray<FName>& OutFormats ) const override;
diff --git a/Engine/Source/Developer/IOS/IOSTargetPlatform/Private/IOSTargetPlatform.cpp b/Engine/Source/Developer/IOS/IOSTargetPlatform/Private/IOSTargetPlatform.cpp
--- a/Engine/Source/Developer/IOS/IOSTargetPlatform/Private/IOSTargetPlatform.cpp
+++ b/Engine/Source/Developer/IOS/IOSTargetPlatform/Private/IOSTargetPlatform.cpp
@@ -13,6 +13,7 @@
 #include "Misc/App.h"
 #include "Misc/MonitoredProcess.h"
 #include "Logging/MessageLog.h"
+#include "AnalyticsEventAttribute.h"
 
 #if PLATFORM_WINDOWS
 #include "Windows/WindowsHWrapper.h"
@@ -529,6 +530,14 @@ void FIOSTargetPlatform::GetAllTargetedShaderFormats( TArray<FName>& OutFormats
 	GetAllPossibleShaderFormats(OutFormats);
 }
 
+void FIOSTargetPlatform::GetPlatformSpecificProjectAnalytics( TArray<FAnalyticsEventAttribute>& AnalyticsParamArray ) const
+{
+	TNonDesktopTargetPlatformBase<FIOSPlatformProperties>::GetPlatformSpecificProjectAnalytics( AnalyticsParamArray );
+
+	AppendAnalyticsEventAttributeArray(AnalyticsParamArray,
+		TEXT("SupportsMetalMRT"), SupportsMetalMRT()
+	);
+}
 
 #if WITH_ENGINE
 
diff --git a/Engine/Source/Developer/IOS/IOSTargetPlatform/Private/IOSTargetPlatform.h b/Engine/Source/Developer/IOS/IOSTargetPlatform/Private/IOSTargetPlatform.h
--- a/Engine/Source/Developer/IOS/IOSTargetPlatform/Private/IOSTargetPlatform.h
+++ b/Engine/Source/Developer/IOS/IOSTargetPlatform/Private/IOSTargetPlatform.h
@@ -65,6 +65,8 @@ public:
 
 	virtual void GetAllTargetedShaderFormats( TArray<FName>& OutFormats ) const override;
 
+	virtual void GetPlatformSpecificProjectAnalytics( TArray<struct FAnalyticsEventAttribute>& AnalyticsParamArray ) const override;
+
 #if WITH_ENGINE
 	virtual void GetReflectionCaptureFormats( TArray<FName>& OutFormats ) const override;
 
diff --git a/Engine/Source/Developer/Linux/LinuxTargetPlatform/Private/LinuxTargetPlatform.h b/Engine/Source/Developer/Linux/LinuxTargetPlatform/Private/LinuxTargetPlatform.h
--- a/Engine/Source/Developer/Linux/LinuxTargetPlatform/Private/LinuxTargetPlatform.h
+++ b/Engine/Source/Developer/Linux/LinuxTargetPlatform/Private/LinuxTargetPlatform.h
@@ -13,6 +13,7 @@
 #include "Interfaces/ITargetPlatform.h"
 #include "Common/TargetPlatformBase.h"
 #include "SteamDeck/SteamDeckDevice.h"
+#include "AnalyticsEventAttribute.h"
 
 
 #if WITH_ENGINE
@@ -293,6 +294,13 @@ public:
 		}
 	}
 
+	virtual void GetPlatformSpecificProjectAnalytics( TArray<FAnalyticsEventAttribute>& AnalyticsParamArray ) const override
+	{
+		TSuper::GetPlatformSpecificProjectAnalytics(AnalyticsParamArray);
+
+		TSuper::AppendAnalyticsEventConfigArray(AnalyticsParamArray, TEXT("/Script/LinuxTargetPlatform.LinuxTargetSettings"), TEXT("TargetedRHIs"), GEngineIni);
+	}
+
 #if WITH_ENGINE
 	virtual void GetReflectionCaptureFormats(TArray<FName>& OutFormats) const override
 	{
diff --git a/Engine/Source/Developer/Mac/MacTargetPlatform/Private/GenericMacTargetPlatform.h b/Engine/Source/Developer/Mac/MacTargetPlatform/Private/GenericMacTargetPlatform.h
--- a/Engine/Source/Developer/Mac/MacTargetPlatform/Private/GenericMacTargetPlatform.h
+++ b/Engine/Source/Developer/Mac/MacTargetPlatform/Private/GenericMacTargetPlatform.h
@@ -11,6 +11,7 @@
 #include "Mac/MacPlatformProperties.h"
 #include "Misc/ConfigCacheIni.h"
 #include "LocalMacTargetDevice.h"
+#include "AnalyticsEventAttribute.h"
 
 
 #if WITH_ENGINE
@@ -146,6 +147,13 @@ public:
 		}
 	}
 
+	virtual void GetPlatformSpecificProjectAnalytics( TArray<FAnalyticsEventAttribute>& AnalyticsParamArray ) const override
+	{
+		TSuper::GetPlatformSpecificProjectAnalytics(AnalyticsParamArray);
+
+		TSuper::AppendAnalyticsEventConfigArray(AnalyticsParamArray, TEXT("/Script/MacTargetPlatform.MacTargetSettings"), TEXT("TargetedRHIs"), GEngineIni);
+	}
+
 #if WITH_ENGINE
 	virtual const class FStaticMeshLODSettings& GetStaticMeshLODSettings( ) const override
 	{
diff --git a/Engine/Source/Developer/TargetPlatform/Private/TargetPlatformBase.cpp b/Engine/Source/Developer/TargetPlatform/Private/TargetPlatformBase.cpp
--- a/Engine/Source/Developer/TargetPlatform/Private/TargetPlatformBase.cpp
+++ b/Engine/Source/Developer/TargetPlatform/Private/TargetPlatformBase.cpp
@@ -11,6 +11,7 @@
 #include "Misc/App.h"
 #include "Misc/ConfigCacheIni.h"
 #include "Sound/AudioFormatSettings.h"
+#include "AnalyticsEventAttribute.h"
 
 #define LOCTEXT_NAMESPACE "TargetPlatform"
 
@@ -104,6 +105,61 @@ void FTargetPlatformBase::GetRayTracingShaderFormats(TArray<FName>& OutFormats)
 	}
 }
 
+void FTargetPlatformBase::GetPlatformSpecificProjectAnalytics( TArray<FAnalyticsEventAttribute>& AnalyticsParamArray ) const
+{
+	AppendAnalyticsEventAttributeArray( AnalyticsParamArray,
+		TEXT("UsesDistanceFields"), UsesDistanceFields(),
+		TEXT("UsesForwardShading"), UsesForwardShading()
+	);
+}
+
+void FTargetPlatformBase::AppendAnalyticsEventConfigBool( TArray<FAnalyticsEventAttribute>& AnalyticsParamArray, const TCHAR* ConfigSection, const TCHAR* ConfigKey, const FString& IniFileName, const TCHAR* AnalyticsKeyNameOverride )
+{
+	bool ConfigValue;
+	if (GConfig->GetBool(ConfigSection, ConfigKey, ConfigValue, IniFileName))
+	{
+		AnalyticsParamArray.Add( FAnalyticsEventAttribute( AnalyticsKeyNameOverride ? AnalyticsKeyNameOverride : ConfigKey, ConfigValue ) );
+	}
+}
+
+void FTargetPlatformBase::AppendAnalyticsEventConfigInt( TArray<FAnalyticsEventAttribute>& AnalyticsParamArray, const TCHAR* ConfigSection, const TCHAR* ConfigKey, const FString& IniFileName, const TCHAR* AnalyticsKeyNameOverride )
+{
+	int32 ConfigValue;
+	if (GConfig->GetInt(ConfigSection, ConfigKey, ConfigValue, IniFileName))
+	{
+		AnalyticsParamArray.Add( FAnalyticsEventAttribute( AnalyticsKeyNameOverride ? AnalyticsKeyNameOverride : ConfigKey, ConfigValue ) );
+	}
+}
+
+void FTargetPlatformBase::AppendAnalyticsEventConfigFloat( TArray<FAnalyticsEventAttribute>& AnalyticsParamArray, const TCHAR* ConfigSection, const TCHAR* ConfigKey, const FString& IniFileName, const TCHAR* AnalyticsKeyNameOverride )
+{
+	float ConfigValue;
+	if (GConfig->GetFloat(ConfigSection, ConfigKey, ConfigValue, IniFileName))
+	{
+		AnalyticsParamArray.Add( FAnalyticsEventAttribute( AnalyticsKeyNameOverride ? AnalyticsKeyNameOverride : ConfigKey, ConfigValue ) );
+	}
+}
+
+void FTargetPlatformBase::AppendAnalyticsEventConfigString( TArray<FAnalyticsEventAttribute>& AnalyticsParamArray, const TCHAR* ConfigSection, const TCHAR* ConfigKey, const FString& IniFileName, const TCHAR* AnalyticsKeyNameOverride )
+{
+	FString ConfigValue;
+	if (GConfig->GetString(ConfigSection, ConfigKey, ConfigValue, IniFileName))
+	{
+		AnalyticsParamArray.Add( FAnalyticsEventAttribute( AnalyticsKeyNameOverride ? AnalyticsKeyNameOverride : ConfigKey, ConfigValue ) );
+	}
+}
+
+void FTargetPlatformBase::AppendAnalyticsEventConfigArray( TArray<FAnalyticsEventAttribute>& AnalyticsParamArray, const TCHAR* ConfigSection, const TCHAR* ConfigKey, const FString& IniFileName, const TCHAR* AnalyticsKeyNameOverride )
+{
+	TArray<FString> ConfigValue;
+	if (GConfig->GetArray(ConfigSection, ConfigKey, ConfigValue, IniFileName))
+	{
+		AnalyticsParamArray.Add( FAnalyticsEventAttribute( AnalyticsKeyNameOverride ? AnalyticsKeyNameOverride : ConfigKey, ConfigValue ) );
+	}
+}
+
+
+
 static bool IsPluginEnabledForTarget(const IPlugin& Plugin, const FProjectDescriptor* Project, const FString& Platform, EBuildConfiguration Configuration, EBuildTargetType TargetType)
 {
 	if (!Plugin.GetDescriptor().SupportsTargetPlatform(Platform))
diff --git a/Engine/Source/Developer/TargetPlatform/Public/Common/TargetPlatformBase.h b/Engine/Source/Developer/TargetPlatform/Public/Common/TargetPlatformBase.h
--- a/Engine/Source/Developer/TargetPlatform/Public/Common/TargetPlatformBase.h
+++ b/Engine/Source/Developer/TargetPlatform/Public/Common/TargetPlatformBase.h
@@ -104,6 +104,8 @@ public:
 
 	TARGETPLATFORM_API virtual void GetRayTracingShaderFormats(TArray<FName>& OutFormats) const override;
 
+	TARGETPLATFORM_API virtual void GetPlatformSpecificProjectAnalytics( TArray<struct FAnalyticsEventAttribute>& AnalyticsParamArray ) const override;
+
 #if WITH_ENGINE
 	virtual void GetReflectionCaptureFormats( TArray<FName>& OutFormats ) const override
 	{
@@ -280,6 +282,14 @@ protected:
 	int32 PlatformOrdinal;
 	
 	TARGETPLATFORM_API const Audio::FAudioFormatSettings& GetAudioFormatSettings() const;
+
+	/** Analytics helper functions */
+	TARGETPLATFORM_API static void AppendAnalyticsEventConfigBool( TArray<struct FAnalyticsEventAttribute>& AnalyticsParamArray, const TCHAR* ConfigSection, const TCHAR* ConfigKey, const FString& IniFileName, const TCHAR* AnalyticsKeyNameOverride = nullptr );
+	TARGETPLATFORM_API static void AppendAnalyticsEventConfigInt( TArray<struct FAnalyticsEventAttribute>& AnalyticsParamArray, const TCHAR* ConfigSection, const TCHAR* ConfigKey, const FString& IniFileName, const TCHAR* AnalyticsKeyNameOverride = nullptr );
+	TARGETPLATFORM_API static void AppendAnalyticsEventConfigFloat( TArray<struct FAnalyticsEventAttribute>& AnalyticsParamArray, const TCHAR* ConfigSection, const TCHAR* ConfigKey, const FString& IniFileName, const TCHAR* AnalyticsKeyNameOverride = nullptr );
+	TARGETPLATFORM_API static void AppendAnalyticsEventConfigString( TArray<struct FAnalyticsEventAttribute>& AnalyticsParamArray, const TCHAR* ConfigSection, const TCHAR* ConfigKey, const FString& IniFileName, const TCHAR* AnalyticsKeyNameOverride = nullptr );
+	TARGETPLATFORM_API static void AppendAnalyticsEventConfigArray( TArray<struct FAnalyticsEventAttribute>& AnalyticsParamArray, const TCHAR* ConfigSection, const TCHAR* ConfigKey, const FString& IniFileName, const TCHAR* AnalyticsKeyNameOverride = nullptr );
+
 private:
 	TPimplPtr<Audio::FAudioFormatSettings> AudioFormatSettings;
 	
diff --git a/Engine/Source/Developer/TargetPlatform/Public/Interfaces/ITargetPlatform.h b/Engine/Source/Developer/TargetPlatform/Public/Interfaces/ITargetPlatform.h
--- a/Engine/Source/Developer/TargetPlatform/Public/Interfaces/ITargetPlatform.h
+++ b/Engine/Source/Developer/TargetPlatform/Public/Interfaces/ITargetPlatform.h
@@ -543,6 +543,11 @@ public:
 	*/
 	virtual void GetRayTracingShaderFormats(TArray<FName>& OutFormats) const = 0;
 
+	/**
+	 * Gather per-project cook/package analytics
+	 */
+	virtual void GetPlatformSpecificProjectAnalytics( TArray<struct FAnalyticsEventAttribute>& AnalyticsParamArray ) const = 0;
+
 #if WITH_ENGINE
 	/**
 	 * Gets the format to use for a particular body setup.
diff --git a/Engine/Source/Developer/TargetPlatform/TargetPlatform.Build.cs b/Engine/Source/Developer/TargetPlatform/TargetPlatform.Build.cs
--- a/Engine/Source/Developer/TargetPlatform/TargetPlatform.Build.cs
+++ b/Engine/Source/Developer/TargetPlatform/TargetPlatform.Build.cs
@@ -15,6 +15,7 @@ public class TargetPlatform : ModuleRules
 		PublicDependencyModuleNames.Add("DeveloperSettings");
 		PublicDependencyModuleNames.Add("AudioPlatformConfiguration");
 		PublicDependencyModuleNames.Add("DesktopPlatform");
+		PublicDependencyModuleNames.Add("Analytics");
 
 		// TextureFormat contains public headers that were historically part of TargetPlatform, so it is exposed
 		// as a public include path on TargetPlatform.
diff --git a/Engine/Source/Developer/Windows/WindowsTargetPlatform/Private/GenericWindowsTargetPlatform.h b/Engine/Source/Developer/Windows/WindowsTargetPlatform/Private/GenericWindowsTargetPlatform.h
--- a/Engine/Source/Developer/Windows/WindowsTargetPlatform/Private/GenericWindowsTargetPlatform.h
+++ b/Engine/Source/Developer/Windows/WindowsTargetPlatform/Private/GenericWindowsTargetPlatform.h
@@ -5,6 +5,7 @@
 #include "CoreMinimal.h"
 #include "Common/TargetPlatformBase.h"
 #include "Misc/ConfigCacheIni.h"
+#include "AnalyticsEventAttribute.h"
 #if PLATFORM_WINDOWS
 #include "LocalPcTargetDevice.h"
 #endif
@@ -285,6 +286,22 @@ public:
 		}
 	}
 
+	virtual void GetPlatformSpecificProjectAnalytics( TArray<FAnalyticsEventAttribute>& AnalyticsParamArray ) const override
+	{
+		TSuper::GetPlatformSpecificProjectAnalytics(AnalyticsParamArray);
+
+		AppendAnalyticsEventAttributeArray(AnalyticsParamArray,
+			TEXT("UsesRayTracing"), UsesRayTracing()
+		);
+	
+		TSuper::AppendAnalyticsEventConfigString(AnalyticsParamArray, TEXT("/Script/WindowsTargetPlatform.WindowsTargetSettings"), TEXT("DefaultGraphicsRHI"), GEngineIni);
+
+		TSuper::AppendAnalyticsEventConfigArray(AnalyticsParamArray, TEXT("/Script/WindowsTargetPlatform.WindowsTargetSettings"), TEXT("D3D12TargetedShaderFormats"), GEngineIni);
+		TSuper::AppendAnalyticsEventConfigArray(AnalyticsParamArray, TEXT("/Script/WindowsTargetPlatform.WindowsTargetSettings"), TEXT("D3D11TargetedShaderFormats"), GEngineIni);
+		TSuper::AppendAnalyticsEventConfigArray(AnalyticsParamArray, TEXT("/Script/WindowsTargetPlatform.WindowsTargetSettings"), TEXT("VulkanTargetedShaderFormats"), GEngineIni);
+		TSuper::AppendAnalyticsEventConfigArray(AnalyticsParamArray, TEXT("/Script/WindowsTargetPlatform.WindowsTargetSettings"), TEXT("TargetedRHIs"), GEngineIni, TEXT("TargetedRHIs_Deprecated") );
+	}
+
 #if WITH_ENGINE
 	virtual void GetReflectionCaptureFormats(TArray<FName>& OutFormats) const override
 	{
diff --git a/Engine/Source/Editor/MainFrame/Private/Frame/MainFrameActions.cpp b/Engine/Source/Editor/MainFrame/Private/Frame/MainFrameActions.cpp
--- a/Engine/Source/Editor/MainFrame/Private/Frame/MainFrameActions.cpp
+++ b/Engine/Source/Editor/MainFrame/Private/Frame/MainFrameActions.cpp
@@ -485,7 +485,7 @@ void FMainFrameActionCallbacks::ZipUpProject()
 			FString CommandLine = FString::Printf(TEXT("ZipProjectUp -nocompileeditor -project=\"%s\" -install=\"%s\""), *ProjectPath, *FinalFileName);
 
 			IUATHelperModule::Get().CreateUatTask( CommandLine, GetTargetPlatformManager()->GetRunningTargetPlatform()->DisplayName(), LOCTEXT("ZipTaskName", "Zipping Up Project"),
-				LOCTEXT("ZipTaskShortName", "Zip Project Task"), FAppStyle::GetBrush(TEXT("MainFrame.CookContent")), IUATHelperModule::UatTaskResultCallack(), FPaths::GetPath(FinalFileName));
+				LOCTEXT("ZipTaskShortName", "Zip Project Task"), FAppStyle::GetBrush(TEXT("MainFrame.CookContent")), nullptr, IUATHelperModule::UatTaskResultCallack(), FPaths::GetPath(FinalFileName));
 		}
 	}
 }
diff --git a/Engine/Source/Editor/TurnkeySupport/Private/TurnkeyEditorSupport.cpp b/Engine/Source/Editor/TurnkeySupport/Private/TurnkeyEditorSupport.cpp
--- a/Engine/Source/Editor/TurnkeySupport/Private/TurnkeyEditorSupport.cpp
+++ b/Engine/Source/Editor/TurnkeySupport/Private/TurnkeyEditorSupport.cpp
@@ -158,7 +158,7 @@ void FTurnkeyEditorSupport::LaunchRunningMap(const FString& DeviceId, const FStr
 				}
 				FText TaskName = LOCTEXT("VerifyingSDK", "Verifying SDK and Device");
 
-				IUATHelperModule::Get().CreateUatTask(CommandLine, FText::FromString(IniPlatformName), TaskName, TaskName, FAppStyle::Get().GetBrush(TEXT("MainFrame.PackageProject")),
+				IUATHelperModule::Get().CreateUatTask(CommandLine, FText::FromString(IniPlatformName), TaskName, TaskName, FAppStyle::Get().GetBrush(TEXT("MainFrame.PackageProject")), nullptr,
 					[SessionParams, RealDeviceName, UBTPlatformName](FString Result, double)
 					{
 						// unfortunate string comparison for success
@@ -220,10 +220,10 @@ bool FTurnkeyEditorSupport::DoesProjectHaveCode()
 #endif
 }
 
-void FTurnkeyEditorSupport::RunUAT(const FString& CommandLine, const FText& PlatformDisplayName, const FText& TaskName, const FText& TaskShortName, const FSlateBrush* TaskIcon, TFunction<void(FString, double)> ResultCallback)
+void FTurnkeyEditorSupport::RunUAT(const FString& CommandLine, const FText& PlatformDisplayName, const FText& TaskName, const FText& TaskShortName, const FSlateBrush* TaskIcon, const TArray<FAnalyticsEventAttribute>* OptionalAnalyticsParamArray, TFunction<void(FString, double)> ResultCallback)
 {
 #if WITH_EDITOR
-	IUATHelperModule::Get().CreateUatTask(CommandLine, PlatformDisplayName, TaskName, TaskShortName, TaskIcon, ResultCallback);
+	IUATHelperModule::Get().CreateUatTask(CommandLine, PlatformDisplayName, TaskName, TaskShortName, TaskIcon, OptionalAnalyticsParamArray, ResultCallback);
 #else
 	unimplemented();
 #endif
diff --git a/Engine/Source/Editor/TurnkeySupport/Private/TurnkeyEditorSupport.h b/Engine/Source/Editor/TurnkeySupport/Private/TurnkeyEditorSupport.h
--- a/Engine/Source/Editor/TurnkeySupport/Private/TurnkeyEditorSupport.h
+++ b/Engine/Source/Editor/TurnkeySupport/Private/TurnkeyEditorSupport.h
@@ -17,7 +17,7 @@ public:
 	
 	static void SaveAll();
 	static bool DoesProjectHaveCode();
-	static void RunUAT(const FString& CommandLine, const FText& PlatformDisplayName, const FText& TaskName, const FText& TaskShortName, const struct FSlateBrush* TaskIcon, TFunction<void(FString, double)> ResultCallback=TFunction<void(FString, double)>());
+	static void RunUAT(const FString& CommandLine, const FText& PlatformDisplayName, const FText& TaskName, const FText& TaskShortName, const struct FSlateBrush* TaskIcon, const TArray<struct FAnalyticsEventAttribute>* OptionalAnalyticsParamArray = nullptr, TFunction<void(FString, double)> ResultCallback=TFunction<void(FString, double)>());
 
 	static bool ShowOKCancelDialog(FText Message, FText Title);
 	static void ShowRestartToast();
diff --git a/Engine/Source/Editor/TurnkeySupport/Private/TurnkeySupportModule.cpp b/Engine/Source/Editor/TurnkeySupport/Private/TurnkeySupportModule.cpp
--- a/Engine/Source/Editor/TurnkeySupport/Private/TurnkeySupportModule.cpp
+++ b/Engine/Source/Editor/TurnkeySupport/Private/TurnkeySupportModule.cpp
@@ -45,6 +45,7 @@
 #include "ToolMenus.h"
 #include "TurnkeyEditorSupport.h"
 #include "ITurnkeyIOModule.h"
+#include "AnalyticsEventAttribute.h"
 #if WITH_EDITOR
 #include "ZenServerInterface.h"
 #endif
@@ -268,6 +269,8 @@ public:
 
 	static void CookOrPackage(FName IniPlatformName, EPrepareContentMode Mode)
 	{
+		TArray<FAnalyticsEventAttribute> AnalyticsParamArray;
+
 		// get a in-memory defaults which will have the user-settings, like the per-platform config/target platform stuff
 		UProjectPackagingSettings* AllPlatformPackagingSettings = GetMutableDefault<UProjectPackagingSettings>();
 	
@@ -376,6 +379,9 @@ public:
 			BuildCookRunParams += TEXT(" -zenstore");
 		}
 
+		// gather analytics
+		const ITargetPlatform* TargetPlatform = GetTargetPlatformManager()->FindTargetPlatform(PlatformInfo->Name);
+		TargetPlatform->GetPlatformSpecificProjectAnalytics( AnalyticsParamArray );
 
 		// per mode settings
 		FText ContentPrepDescription;
@@ -406,7 +412,6 @@ public:
 
 			BuildCookRunParams += TEXT(" -stage -archive -package");
 
-			const ITargetPlatform* TargetPlatform = GetTargetPlatformManager()->FindTargetPlatform(PlatformInfo->Name);
 			if (ShouldBuildProject(PackagingSettings, TargetPlatform))
 			{
 				BuildCookRunParams += TEXT(" -build");
@@ -508,7 +513,6 @@ public:
 		}
 // 		else if (Mode == EPrepareContentMode::PrepareForDebugging)
 // 		{
-// 			const ITargetPlatform* TargetPlatform = GetTargetPlatformManager()->FindTargetPlatform(PlatformInfo->Name);
 // 			if (ShouldBuildProject(PackagingSettings, TargetPlatform))
 // 			{
 // 				BuildCookRunParams += TEXT(" -build");
@@ -546,7 +550,7 @@ public:
 		}
 		CommandLine.Appendf(TEXT("Turnkey %s BuildCookRun %s"), *TurnkeyParams, *BuildCookRunParams);
 
-		FTurnkeyEditorSupport::RunUAT(CommandLine, PlatformInfo->DisplayName, ContentPrepDescription, ContentPrepTaskName, ContentPrepIcon);
+		FTurnkeyEditorSupport::RunUAT(CommandLine, PlatformInfo->DisplayName, ContentPrepDescription, ContentPrepTaskName, ContentPrepIcon, &AnalyticsParamArray);
 	}
 
 	static bool CanExecuteCustomBuild(FName IniPlatformName, FProjectBuildSettings Build)
@@ -862,7 +866,7 @@ static void TurnkeyInstallSdk(FString IniPlatformName, bool bPreferFull, bool bF
 	CommandLine.Appendf(TEXT("Turnkey -command=VerifySdk -UpdateIfNeeded -platform=%s %s %s -noturnkeyvariables -utf8output -WaitForUATMutex"), *IniPlatformName, *OptionalOptions, *ITurnkeyIOModule::Get().GetUATParams());
 
 	FText TaskName = LOCTEXT("InstallingSdk", "Installing Sdk");
-	FTurnkeyEditorSupport::RunUAT(CommandLine, FText::FromString(IniPlatformName), TaskName, TaskName, FAppStyle::GetBrush(TEXT("MainFrame.PackageProject")),
+	FTurnkeyEditorSupport::RunUAT(CommandLine, FText::FromString(IniPlatformName), TaskName, TaskName, FAppStyle::GetBrush(TEXT("MainFrame.PackageProject")), nullptr,
 		[IniPlatformName](FString, double)
 	{
 		AsyncTask(ENamedThreads::GameThread, [IniPlatformName]()
diff --git a/Engine/Source/Editor/TurnkeySupport/TurnkeySupport.Build.cs b/Engine/Source/Editor/TurnkeySupport/TurnkeySupport.Build.cs
--- a/Engine/Source/Editor/TurnkeySupport/TurnkeySupport.Build.cs
+++ b/Engine/Source/Editor/TurnkeySupport/TurnkeySupport.Build.cs
@@ -28,6 +28,7 @@ public class TurnkeySupport : ModuleRules
  				"LauncherServices",
 				"SourceControl",
 				"TurnkeyIO",
+				"Analytics",
 			}
 		);
 
diff --git a/Engine/Source/Editor/UATHelper/Public/IUATHelperModule.h b/Engine/Source/Editor/UATHelper/Public/IUATHelperModule.h
--- a/Engine/Source/Editor/UATHelper/Public/IUATHelperModule.h
+++ b/Engine/Source/Editor/UATHelper/Public/IUATHelperModule.h
@@ -43,6 +43,6 @@ public:
 	/** Creates and starts up a UAT Task
 	  * @param	ResultLocation	The folder where the result of the task will be stored  
 	  */
-	virtual void CreateUatTask(const FString& CommandLine, const FText& PlatformDisplayName, const FText& TaskName, const FText &TaskShortName, const FSlateBrush* TaskIcon, UatTaskResultCallack ResultCallback = UatTaskResultCallack(), const FString& ResultLocation = FString()) = 0;
+	virtual void CreateUatTask(const FString& CommandLine, const FText& PlatformDisplayName, const FText& TaskName, const FText &TaskShortName, const FSlateBrush* TaskIcon, const TArray<struct FAnalyticsEventAttribute>* OptionalAnalyticsParamArray = nullptr, UatTaskResultCallack ResultCallback = UatTaskResultCallack(), const FString& ResultLocation = FString()) = 0;
 };
 
diff --git a/Engine/Source/Editor/UATHelper/UATHelperModule.cpp b/Engine/Source/Editor/UATHelper/UATHelperModule.cpp
--- a/Engine/Source/Editor/UATHelper/UATHelperModule.cpp
+++ b/Engine/Source/Editor/UATHelper/UATHelperModule.cpp
@@ -383,7 +383,7 @@ public:
 	{
 	}
 
-	virtual void CreateUatTask( const FString& CommandLine, const FText& PlatformDisplayName, const FText& TaskName, const FText &TaskShortName, const FSlateBrush* TaskIcon, UatTaskResultCallack ResultCallback, const FString& ResultLocation)
+	virtual void CreateUatTask( const FString& CommandLine, const FText& PlatformDisplayName, const FText& TaskName, const FText &TaskShortName, const FSlateBrush* TaskIcon, const TArray<FAnalyticsEventAttribute>* OptionalAnalyticsParamArray, UatTaskResultCallack ResultCallback, const FString& ResultLocation)
 	{
 		// If this is a packaging or cooking task, clear the PackagingResults log
 		if (!TaskShortName.CompareToCaseIgnored(FText::FromString(TEXT("Packaging"))) || 
@@ -462,7 +462,14 @@ public:
 		}
 
 		FString EventName = (CommandLine.Contains(TEXT("-package")) ? TEXT("Editor.Package") : TEXT("Editor.Cook"));
-		FEditorAnalytics::ReportEvent(EventName + TEXT(".Start"), PlatformDisplayName.ToString(), bHasCode);
+		if (OptionalAnalyticsParamArray != nullptr)
+		{
+			FEditorAnalytics::ReportEvent(EventName + TEXT(".Start"), PlatformDisplayName.ToString(), bHasCode, *OptionalAnalyticsParamArray);
+		}
+		else
+		{
+			FEditorAnalytics::ReportEvent(EventName + TEXT(".Start"), PlatformDisplayName.ToString(), bHasCode);
+		}
 
 		NotificationItem->SetCompletionState(SNotificationItem::CS_Pending);
 
diff --git a/Engine/Source/Editor/UnrealEd/Private/PlayLevelViaLauncher.cpp b/Engine/Source/Editor/UnrealEd/Private/PlayLevelViaLauncher.cpp
--- a/Engine/Source/Editor/UnrealEd/Private/PlayLevelViaLauncher.cpp
+++ b/Engine/Source/Editor/UnrealEd/Private/PlayLevelViaLauncher.cpp
@@ -425,8 +425,10 @@ void UEditorEngine::StartPlayUsingLauncherSession(FRequestPlaySessionParams& InR
 	}
 
 	// analytics for launch on
-	int32 ErrorCode = 0;
-	FEditorAnalytics::ReportEvent(TEXT("Editor.LaunchOn.Started"), LastPlayUsingLauncherDeviceId.Left(LastPlayUsingLauncherDeviceId.Find(TEXT("@"))), LauncherSessionInfo->bPlayUsingLauncherHasCode);
+	TArray<FAnalyticsEventAttribute> AnalyticsParamArray;
+	LaunchPlatform->GetPlatformSpecificProjectAnalytics(AnalyticsParamArray);
+	FEditorAnalytics::ReportEvent(TEXT("Editor.LaunchOn.Started"), LaunchPlatformName, LauncherSessionInfo->bPlayUsingLauncherHasCode, AnalyticsParamArray);
+
 
 	NotificationItem->SetCompletionState(SNotificationItem::CS_Pending);
 
