Instrument loading times of PS4 & Switch target platform modules

RESOLUTION: fixed conflict with new "impl" method

diff --git a/Engine/Source/Developer/TargetPlatform/Private/TargetPlatformManagerModule.cpp b/Engine/Source/Developer/TargetPlatform/Private/TargetPlatformManagerModule.cpp
--- a/Engine/Source/Developer/TargetPlatform/Private/TargetPlatformManagerModule.cpp
+++ b/Engine/Source/Developer/TargetPlatform/Private/TargetPlatformManagerModule.cpp
@@ -674,6 +674,8 @@ protected:
 
 	bool InitializeSinglePlatform(FName PlatformName, const FString& AutoSDKPath)
 	{
+		TRACE_CPUPROFILER_EVENT_SCOPE(FTargetPlatformManagerModule::InitializeSinglePlatform);
+
 		// try the incoming name as a module name, or as a platform name
 		FName PlatformModuleName = PlatformName;
 
@@ -692,6 +694,8 @@ protected:
 		// original logic for module loading here
 		if (Module)
 		{
+			TRACE_CPUPROFILER_EVENT_SCOPE(FTargetPlatformManagerModule::InitializeSinglePlatform Loading);
+
 			// would like to move this check to GetActiveTargetPlatforms, but too many things cache this result
 			// this setup will become faster after TTP 341897 is complete.
 		RETRY_SETUPANDVALIDATE:
@@ -740,11 +744,15 @@ protected:
 		// Find all module subdirectories and add them so we can load dependent modules for target platform modules
 		// We may not be able to restrict this to subdirectories found in FPlatformInfo because we could have a subdirectory
 		// that is not one of these platforms. Imagine a "Sega" shared directory for the "Genesis" and "Dreamcast" platforms
-		TArray<FString> ModuleSubdirs;
-		IFileManager::Get().FindFilesRecursive(ModuleSubdirs, *FPlatformProcess::GetModulesDirectory(), TEXT("*"), false, true);
-		for (const FString& ModuleSubdir : ModuleSubdirs)
 		{
-			FModuleManager::Get().AddBinariesDirectory(*ModuleSubdir, false);
+			TRACE_CPUPROFILER_EVENT_SCOPE(FTargetPlatformManagerModule::DiscoverAvailablePlatforms FindFilesRecursive);
+
+			TArray<FString> ModuleSubdirs;
+			IFileManager::Get().FindFilesRecursive(ModuleSubdirs, *FPlatformProcess::GetModulesDirectory(), TEXT("*"), false, true);
+			for (const FString& ModuleSubdir : ModuleSubdirs)
+			{
+				FModuleManager::Get().AddBinariesDirectory(*ModuleSubdir, false);
+			}
 		}
 #endif
 
@@ -801,6 +809,8 @@ protected:
 
 	bool SetupAndValidateAutoSDK(const FString& AutoSDKPath)
 	{
+		TRACE_CPUPROFILER_EVENT_SCOPE(FTargetPlatformManagerModule::SetupAndValidateAutoSDK);
+
 #if AUTOSDKS_ENABLED
 		bool bValidSDK = false;
 		if (AutoSDKPath.Len() > 0)
@@ -831,7 +841,9 @@ protected:
 	}
 	
 	bool SetupEnvironmentFromAutoSDK(const FString& AutoSDKPath)
-	{						
+	{
+		TRACE_CPUPROFILER_EVENT_SCOPE(FTargetPlatformManagerModule::SetupEnvironmentFromAutoSDK);
+
 #if AUTOSDKS_ENABLED
 		
 		if (!UE::AutoSDK::IsAutoSDKsEnabled())
@@ -1025,6 +1037,8 @@ protected:
 
 	bool SetupSDKStatus(const FString& TargetPlatforms)
 	{
+		TRACE_CPUPROFILER_EVENT_SCOPE(FTargetPlatformManagerModule::SetupSDKStatus);
+
 //		FDataDrivenPlatformInfoRegistry::UpdateSdkStatus();
 #if 0
 		DECLARE_SCOPE_CYCLE_COUNTER( TEXT( "FTargetPlatformManagerModule::SetupSDKStatus" ), STAT_FTargetPlatformManagerModule_SetupSDKStatus, STATGROUP_TargetPlatform );
@@ -1044,7 +1058,7 @@ protected:
 		SDKStatusMessage = TEXT("");
 
 		{
-			TRACE_CPUPROFILER_EVENT_SCOPE(WaitUntilUBTStarted);
+			TRACE_CPUPROFILER_EVENT_SCOPE(FTargetPlatformManagerModule::SetupSDKStatus WaitUntilUBTStarted);
 			UBTProcess->Launch();
 			while(UBTProcess->Update())
 			{
@@ -1119,6 +1133,8 @@ protected:
 
 	bool UpdateAfterSDKInstall(FName PlatformName)
 	{
+		TRACE_CPUPROFILER_EVENT_SCOPE(FTargetPlatformManagerModule::UpdateAfterSDKInstall);
+
 		const FDataDrivenPlatformInfo& Info = FDataDrivenPlatformInfoRegistry::GetPlatformInfo(PlatformName);
 
 #if AUTOSDKS_ENABLED
