- Include the VVM format to ensure we find the format
- Without doing this the found formats might not contain the format leading to a crash when compiling with the VVM

#rnx

diff --git a/Engine/Source/Developer/TargetPlatform/Private/TargetPlatformManagerModule.cpp b/Engine/Source/Developer/TargetPlatform/Private/TargetPlatformManagerModule.cpp
--- a/Engine/Source/Developer/TargetPlatform/Private/TargetPlatformManagerModule.cpp
+++ b/Engine/Source/Developer/TargetPlatform/Private/TargetPlatformManagerModule.cpp
@@ -571,6 +571,7 @@ public:
 		static void GetHintedModules(ITargetPlatform* Platform, TArray<FName>& Hints)
 		{
 			Platform->GetShaderFormatModuleHints(Hints);
+			Hints.Add(TEXT("ShaderFormatVectorVM"));
 		}
 		static void GetRequiredFormats(ITargetPlatform* Platform, TArray<FName>& RequiredFormats)
 		{
