Fix for editor crash on startup

diff --git a/Engine/Source/Developer/TargetPlatform/Private/TargetPlatformManagerModule.cpp b/Engine/Source/Developer/TargetPlatform/Private/TargetPlatformManagerModule.cpp
--- a/Engine/Source/Developer/TargetPlatform/Private/TargetPlatformManagerModule.cpp
+++ b/Engine/Source/Developer/TargetPlatform/Private/TargetPlatformManagerModule.cpp
@@ -480,7 +480,7 @@ public:
 					if (Module != nullptr)
 					{
 						FormatType* Format = HelperType::GetFormatFromModule(Module);
-						if (!Results.Contains(Format))
+						if (Format && !Results.Contains(Format))
 						{
 							// remember the module
 							Results.Add(Format);
