First part of TargetPlatform(TP) refactor into TargetPlatformSettings(TPS) and TargetPlatformControls(TPC)
TPS doesn't need SDK and will be used to get info about the platform
TPC requires SDK
AndroidTP has been converted to the new system

diff --git a/Engine/Source/Developer/Android/AndroidTargetPlatform/Private/AndroidTargetDevice.h b/Engine/Source/Developer/Android/AndroidTargetPlatform/Private/AndroidTargetDevice.h
deleted file mode 100644
diff --git a/Engine/Source/Developer/Android/AndroidTargetPlatform/Private/AndroidTargetDevice.inl b/Engine/Source/Developer/Android/AndroidTargetPlatform/Private/AndroidTargetDevice.inl
deleted file mode 100644
diff --git a/Engine/Source/Developer/Android/AndroidTargetPlatform/Private/AndroidTargetDeviceOutput.h b/Engine/Source/Developer/Android/AndroidTargetPlatform/Private/AndroidTargetDeviceOutput.h
deleted file mode 100644
diff --git a/Engine/Source/Developer/Android/AndroidTargetPlatform/Private/AndroidTargetDeviceOutput.inl b/Engine/Source/Developer/Android/AndroidTargetPlatform/Private/AndroidTargetDeviceOutput.inl
deleted file mode 100644
diff --git a/Engine/Source/Developer/Android/AndroidTargetPlatform/Private/AndroidTargetPlatform.cpp b/Engine/Source/Developer/Android/AndroidTargetPlatform/Private/AndroidTargetPlatform.cpp
deleted file mode 100644
diff --git a/Engine/Source/Developer/Android/AndroidTargetPlatform/Private/AndroidTargetPlatform.h b/Engine/Source/Developer/Android/AndroidTargetPlatform/Private/AndroidTargetPlatform.h
deleted file mode 100644
diff --git a/Engine/Source/Developer/Android/AndroidTargetPlatform/Private/AndroidTargetPlatformModule.cpp b/Engine/Source/Developer/Android/AndroidTargetPlatform/Private/AndroidTargetPlatformModule.cpp
--- a/Engine/Source/Developer/Android/AndroidTargetPlatform/Private/AndroidTargetPlatformModule.cpp
+++ b/Engine/Source/Developer/Android/AndroidTargetPlatform/Private/AndroidTargetPlatformModule.cpp
@@ -7,71 +7,26 @@
 #include "Common/TargetPlatformBase.h"
 #include "Interfaces/IAndroidDeviceDetection.h"
 #include "Interfaces/IAndroidDeviceDetectionModule.h"
-#include "AndroidTargetDevice.h"
-#include "AndroidTargetPlatform.h"
 #include "IAndroidTargetPlatformModule.h"
 
-#define LOCTEXT_NAMESPACE "FAndroidTargetPlatformModule"
-
 /**
  * Module for the Android target platform.
  */
-class FAndroidTargetPlatformModule : public IAndroidTargetPlatformModule
+class FAndroidTargetPlatformModule : public ITargetPlatformModule
 {
 public:
-
-	virtual void GetTargetPlatforms(TArray<ITargetPlatform*>& TargetPlatforms) override
+	virtual void GetTargetPlatforms(TArray<ITargetPlatform*>& TargetPlatforms)
 	{
-		if (FAndroidTargetPlatform::IsUsable())
-		{
-			 // ensure this is only called once
-			check( SinglePlatforms.IsEmpty() );
-
-			for (int32 Type = 0; Type < 2; Type++)
-			{
-				bool bIsClient = Type == 1;
 
-				// flavorless must come first
-				// "flavorless" (aka "vanilla") Android is not well defined and not actually usable
-				//  but Editor wants to see it in the platform list
-				SinglePlatforms.Add(new FAndroidTargetPlatform(bIsClient, nullptr));
-				SinglePlatforms.Add(new FAndroid_ASTCTargetPlatform(bIsClient));
-				SinglePlatforms.Add(new FAndroid_DXTTargetPlatform(bIsClient));
-				SinglePlatforms.Add(new FAndroid_ETC2TargetPlatform(bIsClient));
-
-				// these are used in NotifyMultiSelectedFormatsChanged, so track in another array
-				MultiPlatforms.Add(new FAndroid_MultiTargetPlatform(bIsClient));
-			}
-
-			// join the single and the multi into one
-			TargetPlatforms.Append(SinglePlatforms);
-			TargetPlatforms.Append(MultiPlatforms);
-
-			// set up the multi platforms now that we have all the other platforms ready to go
-			NotifyMultiSelectedFormatsChanged();
-		}
 	}
 
-
-	virtual void NotifyMultiSelectedFormatsChanged() override
+	virtual void GetTargetPlatforms(TArray<ITargetPlatform*>& TargetPlatforms, TArray<ITargetPlatformSettings*> TargetPlatformSettings, TArray<ITargetPlatformControls*> TargetPlatformControls) override
 	{
-		for (FAndroid_MultiTargetPlatform* TP : MultiPlatforms)
+		for (ITargetPlatformControls* TargetPlatformControlsIt : TargetPlatformControls)
 		{
-			TP->LoadFormats(SinglePlatforms);
+			TargetPlatforms.Add(new FTargetPlatformMerged(TargetPlatformControlsIt->GetTargetPlatformSettings(), TargetPlatformControlsIt));
 		}
-		// @todo multi needs to be passed this event!
 	}
-
-
-private:
-
-	/** Holds the specific types of target platforms for NotifyMultiSelectedFormatsChanged */
-	TArray<FAndroidTargetPlatform*> SinglePlatforms;
-	TArray<FAndroid_MultiTargetPlatform*> MultiPlatforms;
 };
 
-
-#undef LOCTEXT_NAMESPACE
-
-
 IMPLEMENT_MODULE( FAndroidTargetPlatformModule, AndroidTargetPlatform);
diff --git a/Engine/Source/Developer/Android/AndroidTargetPlatformControls/AndroidTargetPlatformControls.Build.cs b/Engine/Source/Developer/Android/AndroidTargetPlatformControls/AndroidTargetPlatformControls.Build.cs
new file mode 100644
--- /dev/null
+++ b/Engine/Source/Developer/Android/AndroidTargetPlatformControls/AndroidTargetPlatformControls.Build.cs
@@ -0,0 +1,40 @@
+// Copyright Epic Games, Inc. All Rights Reserved.
+
+using UnrealBuildTool;
+
+public class AndroidTargetPlatformControls : ModuleRules
+{
+	public AndroidTargetPlatformControls(ReadOnlyTargetRules Target) : base(Target)
+	{
+		UnsafeTypeCastWarningLevel = WarningLevel.Error;
+		BinariesSubFolder = "Android";
+
+		PrivateDependencyModuleNames.AddRange(
+			new string[]
+			{
+				"Core",
+				"CoreUObject",
+				"TargetPlatform",
+				"DesktopPlatform",
+				"AndroidDeviceDetection",
+				"AudioPlatformConfiguration",
+				"AndroidTargetPlatformSettings"
+			}
+		);
+
+		PrivateIncludePathModuleNames.AddRange(
+			new string[]
+			{
+				"AndroidTargetPlatformSettings"
+			}
+			);
+
+		if (Target.bCompileAgainstEngine)
+		{
+			PrivateDependencyModuleNames.Add("Engine");
+			PrivateIncludePathModuleNames.Add("TextureCompressor");     //@todo android: AndroidTargetPlatform.Build
+		}
+
+		PublicDefinitions.Add("WITH_OGGVORBIS=1");
+	}
+}
diff --git a/Engine/Source/Developer/Android/AndroidTargetPlatformControls/Private/AndroidTargetDevice.h b/Engine/Source/Developer/Android/AndroidTargetPlatformControls/Private/AndroidTargetDevice.h
new file mode 100644
--- /dev/null
+++ b/Engine/Source/Developer/Android/AndroidTargetPlatformControls/Private/AndroidTargetDevice.h
@@ -0,0 +1,262 @@
+// Copyright Epic Games, Inc. All Rights Reserved.
+
+/*=============================================================================
+	AndroidTargetDevice.h: Declares the AndroidTargetDevice class.
+=============================================================================*/
+
+#pragma once
+
+#include "CoreTypes.h"
+#include "Containers/UnrealString.h"
+#include "Interfaces/ITargetDevice.h"
+#include "Interfaces/ITargetPlatform.h"
+#include "Templates/SharedPointer.h"
+#include "CoreMinimal.h"
+#include "HAL/PlatformProcess.h"
+
+class FAndroidTargetDevice;
+class FTargetDeviceId;
+class ITargetPlatform;
+struct FTargetDeviceProcessInfo;
+enum class ETargetDeviceFeatures;
+enum class ETargetDeviceTypes;
+
+/**
+ * Type definition for shared pointers to instances of FAndroidTargetDevice.
+ */
+typedef TSharedPtr<class FAndroidTargetDevice, ESPMode::ThreadSafe> FAndroidTargetDevicePtr;
+
+/**
+ * Type definition for shared references to instances of FAndroidTargetDevice.
+ */
+typedef TSharedRef<class FAndroidTargetDevice, ESPMode::ThreadSafe> FAndroidTargetDeviceRef;
+
+/**
+ * Type definition for shared references to instances of FAndroidTargetDevice.
+ */
+typedef TSharedPtr<class FAndroidTargetDeviceOutput, ESPMode::ThreadSafe> FAndroidTargetDeviceOutputPtr;
+
+/**
+ * Implements a Android target device.
+ */
+class FAndroidTargetDevice : public ITargetDevice
+{
+public:
+
+	/**
+	 * Creates and initializes a new Android target device.
+	 *
+	 * @param InTargetPlatformControls - The target platform controls.
+	 * @param InSerialNumber - The ADB serial number of the target device.
+	 * @param InAndroidVariant - The variant of the Android platform, i.e. ETC2, DXT or ASTC.
+	 */
+	FAndroidTargetDevice(const ITargetPlatformControls& InTargetPlatformControls, const FString& InSerialNumber, const FString& InAndroidVariant)
+		: AndroidVariant(InAndroidVariant)
+		, bConnected(false)
+		, bIsDeviceAuthorized(false)
+		, AndroidSDKVersion(INDEX_NONE)
+		, DeviceName(InSerialNumber)
+		, Model(InSerialNumber)
+		, SerialNumber(InSerialNumber)
+		, TargetPlatformControls(InTargetPlatformControls)
+	{ }
+
+public:
+
+	/**
+	 * Sets the device's connection state.
+	 *
+	 * @param bInConnected - Whether the device is connected.
+	 */
+	void SetConnected(bool bInConnected)
+	{
+		bConnected = bInConnected;
+	}
+
+	/**
+	 * Sets the device's authorization state.
+	 *
+	 * @param bInConnected - Whether the device is authorized for USB communications.
+	 */
+	void SetAuthorized(bool bInIsAuthorized)
+	{
+		bIsDeviceAuthorized = bInIsAuthorized;
+	}
+
+	/**
+	 * Sets the device's OS/SDK versions.
+	 *
+	 * @param InSDKVersion - Android SDK version of the device.
+	 * @param InReleaseVersion - Android Release (human-readable) version of the device.
+	 */
+	void SetVersions(int32 InSDKVersion, const FString& InReleaseVersion)
+	{
+		AndroidSDKVersion = InSDKVersion;
+		AndroidVersionString = InReleaseVersion;
+	}
+
+	/**
+	 * Sets the device name.
+	 *
+	 * @param InDeviceName - The device name to set.
+	 */
+	void SetDeviceName(const FString& InDeviceName)
+	{
+		DeviceName = InDeviceName;
+	}
+
+	/**
+	 * Sets the device name.
+	 *
+	 * @param InDeviceName - The device name to set.
+	 */
+	void SetModel(const FString& InNodel)
+	{
+		Model = InNodel;
+	}
+
+	FString GetSerialNumber() const
+	{
+		return SerialNumber;
+	}
+
+public:
+
+	//~ Begin ITargetDevice Interface
+	virtual bool Connect() override
+	{
+		return true;
+	}
+
+	virtual void Disconnect() override
+	{
+	}
+
+	virtual ETargetDeviceTypes GetDeviceType() const override
+	{
+		//@TODO: How to distinguish between a Tablet and a Phone (or a TV microconsole, etc...), and is it important?
+		return ETargetDeviceTypes::Tablet;
+	}
+
+	virtual FTargetDeviceId GetId() const override
+	{
+		return FTargetDeviceId(TargetPlatformControls.PlatformName(), SerialNumber);
+	}
+
+	virtual FString GetName() const override
+	{
+		// we need a unique name for all devices, so use human usable model name and the unique id
+		return FString::Printf(TEXT("%s (%s)"), *Model, *SerialNumber);
+	}
+
+	virtual FString GetModelId() const override
+	{
+		return Model;
+	}
+
+	virtual FString GetOSVersion() const override
+	{
+		return AndroidVersionString;
+	}
+
+	virtual FString GetOperatingSystemName() override;
+
+	virtual int32 GetProcessSnapshot( TArray<FTargetDeviceProcessInfo>& OutProcessInfos ) override;
+
+	virtual bool IsConnected() override
+	{
+		return bConnected;
+	}
+
+	virtual bool IsDefault() const override
+	{
+		return true;
+	}
+
+	virtual bool IsAuthorized() const override
+	{
+		return bIsDeviceAuthorized;
+	}
+
+	virtual bool PowerOff(bool Force) override;
+
+	virtual bool PowerOn() override
+	{
+		return true;
+	}
+
+	// Return true if the devices can be grouped in an aggregate (All_<platform>_devices_on_<host>) proxy
+	virtual bool IsPlatformAggregated() const override
+	{
+		return true;
+	}
+
+	// the name of the aggregate (All_<platform>_devices_on_<host>) proxy
+	virtual FString GetAllDevicesName() const override;
+
+	// the default variant (texture compression) of the aggregate (All_<platform>_devices_on_<host>) proxy
+	virtual FName GetAllDevicesDefaultVariant() const override
+	{
+		// The Android platform has an aggregate (All_<platform>_devices_on_<host>) entry in the Project Launcher
+		// Multi is the default texture format
+		return "Android_Multi";
+	}
+
+	virtual bool Reboot(bool bReconnect = false) override;
+	virtual bool TerminateLaunchedProcess(const FString& ProcessIdentifier) override;
+	virtual bool SupportsFeature(ETargetDeviceFeatures Feature) const override;
+	virtual bool TerminateProcess(const int64 ProcessId) override;
+	virtual void SetUserCredentials(const FString& UserName, const FString& UserPassword) override;
+	virtual bool GetUserCredentials(FString& OutUserName, FString& OutUserPassword) override;
+	virtual void ExecuteConsoleCommand(const FString& ExecCommand) const override;
+	virtual ITargetDeviceOutputPtr CreateDeviceOutputRouter(FOutputDevice* Output) const override;
+	//~ End ITargetDevice Interface
+
+	/** Full filename for ADB executable. */
+	static bool GetAdbFullFilename(FString& OutFilename);
+
+protected:
+
+	/**
+	 * Executes an SDK command with the specified command line on this device only using ADB.
+	 *
+	 * @param Params - The command line parameters.
+	 * @param OutStdOut - Optional pointer to a string that will hold the command's output log.
+	 * @param OutStdErr - Optional pointer to a string that will hold the error message, if any.
+	 *
+	 * @return true on success, false otherwise.
+	 */
+	bool ExecuteAdbCommand( const FString& Params, FString* OutStdOut, FString* OutStdErr ) const;
+
+protected:
+
+	// The variant of the Android platform, i.e. ETC2, DXT or ASTC.
+	FString AndroidVariant;
+
+	// Holds a flag indicating whether the device is currently connected.
+	bool bConnected;
+
+	// Holds a flag indicating whether the device is USB comms authorized (if not, most other values aren't valid but we still want to show the device as detected but unready)
+	bool bIsDeviceAuthorized;
+
+	// Holds the Android SDK version
+	int32 AndroidSDKVersion;
+
+	// Holds the Android Release version string (e.g., "2.3" or "4.2.2")
+	FString AndroidVersionString;
+
+	// Holds the device name.
+	FString DeviceName;
+
+	// Holds the device model.
+	FString Model;
+
+	// Holds the serial number (from ADB devices) of this target device.
+	FString SerialNumber;
+
+	// Holds a reference to the device's target platform.
+	const ITargetPlatformControls& TargetPlatformControls;
+};
+
+
+#include "AndroidTargetDevice.inl" // IWYU pragma: export
diff --git a/Engine/Source/Developer/Android/AndroidTargetPlatformControls/Private/AndroidTargetDevice.inl b/Engine/Source/Developer/Android/AndroidTargetPlatformControls/Private/AndroidTargetDevice.inl
new file mode 100644
--- /dev/null
+++ b/Engine/Source/Developer/Android/AndroidTargetPlatformControls/Private/AndroidTargetDevice.inl
@@ -0,0 +1,181 @@
+// Copyright Epic Games, Inc. All Rights Reserved.
+
+// HEADER_UNIT_SKIP - Not included directly
+
+/* ITargetDevice interface
+ *****************************************************************************/
+#include "CoreTypes.h"
+#include "Containers/UnrealString.h"
+#include "Misc/Optional.h"
+#include "HAL/PlatformProcess.h"
+#include "AndroidTargetDeviceOutput.h"
+
+class FAndroidTargetDevice;
+struct FTargetDeviceProcessInfo;
+enum class ETargetDeviceFeatures;
+
+inline FString FAndroidTargetDevice::GetOperatingSystemName()
+{
+	if (!AndroidVersionString.IsEmpty())
+	{
+		return FString::Printf(TEXT("Android %s, API level %d"), *AndroidVersionString, AndroidSDKVersion);
+	}
+	else
+	{
+		return TEXT("Android");
+	}
+}
+
+inline int32 FAndroidTargetDevice::GetProcessSnapshot( TArray<FTargetDeviceProcessInfo>& OutProcessInfos ) 
+{
+	return 0;
+}
+
+
+inline bool FAndroidTargetDevice::Reboot( bool bReconnect )
+{
+	if (!ExecuteAdbCommand(TEXT("reboot"), NULL, NULL))
+	{
+		return false;
+	}
+
+	return true;
+}
+
+inline bool FAndroidTargetDevice::PowerOff( bool Force )
+{
+	if (!ExecuteAdbCommand(TEXT("reboot --poweroff"), NULL, NULL))
+	{
+		return false;
+	}
+
+	return true;
+}
+
+inline FString FAndroidTargetDevice::GetAllDevicesName() const
+{
+	return FString::Printf(TEXT("All_%s_On_%s"), *(GetTargetPlatform().IniPlatformName()), FPlatformProcess::ComputerName());
+}
+
+// cancel the running application
+inline bool FAndroidTargetDevice::TerminateLaunchedProcess(const FString& ProcessIdentifier)
+{
+	FString AdbCommand = FString::Printf(TEXT("shell am force-stop '%s' -s %s"), *ProcessIdentifier, *SerialNumber);
+	return ExecuteAdbCommand(AdbCommand, nullptr, nullptr);
+}
+
+inline bool FAndroidTargetDevice::SupportsFeature( ETargetDeviceFeatures Feature ) const
+{
+	switch (Feature)
+	{
+	case ETargetDeviceFeatures::PowerOff:
+		return true;
+
+	case ETargetDeviceFeatures::PowerOn:
+		return false;
+
+	case ETargetDeviceFeatures::Reboot:
+		return true;
+
+	default:
+		return false;
+	}
+}
+
+
+inline bool FAndroidTargetDevice::TerminateProcess( const int64 ProcessId )
+{
+	return false;
+}
+
+
+inline void FAndroidTargetDevice::SetUserCredentials( const FString& UserName, const FString& UserPassword )
+{
+}
+
+
+inline bool FAndroidTargetDevice::GetUserCredentials( FString& OutUserName, FString& OutUserPassword )
+{
+	return false;
+}
+
+inline void FAndroidTargetDevice::ExecuteConsoleCommand(const FString& ExecCommand) const
+{
+	FString AdbCommand = FString::Printf(TEXT("shell \"am broadcast -a android.intent.action.RUN -e cmd '%s'\""), *ExecCommand);
+	ExecuteAdbCommand(AdbCommand, nullptr, nullptr);
+}
+
+inline ITargetDeviceOutputPtr FAndroidTargetDevice::CreateDeviceOutputRouter(FOutputDevice* Output) const
+{
+	FAndroidTargetDeviceOutputPtr DeviceOutputPtr = MakeShareable(new FAndroidTargetDeviceOutput());
+	if (DeviceOutputPtr->Init(*this, Output))
+	{
+		return DeviceOutputPtr;
+	}
+
+	return nullptr;
+}
+
+/* FAndroidTargetDevice implementation
+ *****************************************************************************/
+
+inline bool FAndroidTargetDevice::GetAdbFullFilename(FString& OutFilename)
+{
+	TOptional<FString> ResultPath;
+
+	// get the SDK binaries folder
+	FString AndroidDirectory = FPlatformMisc::GetEnvironmentVariable(TEXT("ANDROID_HOME"));
+	if (AndroidDirectory.Len() == 0)
+	{
+		return false;
+	}
+
+#if PLATFORM_WINDOWS
+	OutFilename = FString::Printf(TEXT("%s\\platform-tools\\adb.exe"), *AndroidDirectory);
+#else
+	OutFilename = FString::Printf(TEXT("%s/platform-tools/adb"), *AndroidDirectory);
+#endif
+
+	return true;
+}
+
+inline bool FAndroidTargetDevice::ExecuteAdbCommand( const FString& CommandLine, FString* OutStdOut, FString* OutStdErr ) const
+{
+	FString Filename;
+	if (!GetAdbFullFilename(Filename))
+	{
+		return false;
+	}
+	
+	// execute the command
+	int32 ReturnCode;
+	FString DefaultError;
+
+	// make sure there's a place for error output to go if the caller specified NULL
+	if (!OutStdErr)
+	{
+		OutStdErr = &DefaultError;
+	}
+
+	FString CommandLineWithDevice;
+	// the devices command should never include a specific device
+	if (CommandLine == TEXT("devices"))
+	{
+		CommandLineWithDevice = CommandLine;
+	}
+	else
+	{
+		CommandLineWithDevice = FString::Printf(TEXT("-s %s %s"), *SerialNumber, *CommandLine);
+	}
+
+	FPlatformProcess::ExecProcess(*Filename, *CommandLineWithDevice, &ReturnCode, OutStdOut, OutStdErr);
+
+	if (ReturnCode != 0)
+	{
+		FPlatformMisc::LowLevelOutputDebugStringf(TEXT("The Android SDK command '%s' failed to run. Return code: %d, Error: %s\n"), *CommandLineWithDevice, ReturnCode, **OutStdErr);
+
+		return false;
+	}
+
+	return true;
+}
diff --git a/Engine/Source/Developer/Android/AndroidTargetPlatformControls/Private/AndroidTargetDeviceOutput.h b/Engine/Source/Developer/Android/AndroidTargetPlatformControls/Private/AndroidTargetDeviceOutput.h
new file mode 100644
--- /dev/null
+++ b/Engine/Source/Developer/Android/AndroidTargetPlatformControls/Private/AndroidTargetDeviceOutput.h
@@ -0,0 +1,58 @@
+// Copyright Epic Games, Inc. All Rights Reserved.
+
+#pragma once
+
+// HEADER_UNIT_SKIP - Not included directly
+
+#include "Containers/UnrealString.h"
+#include "HAL/Runnable.h"
+#include "HAL/RunnableThread.h"
+#include "HAL/PlatformProcess.h"
+#include "Interfaces/ITargetDeviceOutput.h"
+#include "Misc/ConfigCacheIni.h"
+#include "HAL/ThreadSafeCounter.h"
+
+class FAndroidTargetDevice;
+
+class FAndroidDeviceOutputReaderRunnable : public FRunnable
+{
+public:
+	FAndroidDeviceOutputReaderRunnable(const FString& AdbFilename, const FString& DeviceSerialNumber, FOutputDevice* Output);
+	
+	// FRunnable interface.
+	virtual bool Init(void) override;
+	virtual void Exit(void) override; 
+	virtual void Stop(void) override;
+	virtual uint32 Run(void) override;
+
+private:
+	bool StartLogcatProcess(void);
+
+private:
+	// > 0 if we've been asked to abort work in progress at the next opportunity
+	FThreadSafeCounter	StopTaskCounter;
+	
+	FString				AdbFilename;
+	FString				DeviceSerialNumber;
+	FOutputDevice*		Output;
+	
+	void*				LogcatReadPipe;
+	void*				LogcatWritePipe;
+	FProcHandle			LogcatProcHandle;
+};
+
+/**
+ * Implements a Android target device.
+ */
+class FAndroidTargetDeviceOutput : public ITargetDeviceOutput
+{
+public:
+	bool Init(const FAndroidTargetDevice& TargetDevice, FOutputDevice* Output);
+	
+private:
+	TUniquePtr<FRunnableThread>						DeviceOutputThread;
+	FString											DeviceSerialNumber;
+	FString											DeviceName;
+};
+
+#include "AndroidTargetDeviceOutput.inl" // IWYU pragma: export
diff --git a/Engine/Source/Developer/Android/AndroidTargetPlatformControls/Private/AndroidTargetDeviceOutput.inl b/Engine/Source/Developer/Android/AndroidTargetPlatformControls/Private/AndroidTargetDeviceOutput.inl
new file mode 100644
--- /dev/null
+++ b/Engine/Source/Developer/Android/AndroidTargetPlatformControls/Private/AndroidTargetDeviceOutput.inl
@@ -0,0 +1,123 @@
+// Copyright Epic Games, Inc. All Rights Reserved.
+
+// HEADER_UNIT_SKIP - Not included directly
+
+#include "CoreTypes.h"
+#include "CoreFwd.h"
+
+class FAndroidDeviceOutputReaderRunnable;
+class FAndroidTargetDevice;
+class FAndroidTargetDeviceOutput;
+
+inline FAndroidDeviceOutputReaderRunnable::FAndroidDeviceOutputReaderRunnable(const FString& InAdbFilename, const FString& InDeviceSerialNumber, FOutputDevice* InOutput) 
+	: StopTaskCounter(0)
+	, AdbFilename(InAdbFilename)
+	, DeviceSerialNumber(InDeviceSerialNumber)
+	, Output(InOutput)
+	, LogcatReadPipe(nullptr)
+	, LogcatWritePipe(nullptr)
+{
+}
+
+inline bool FAndroidDeviceOutputReaderRunnable::StartLogcatProcess(void)
+{
+	FString Params = FString::Printf(TEXT(" -s %s logcat UE:V DEBUG:V *:S -v tag"), *DeviceSerialNumber);
+	LogcatProcHandle = FPlatformProcess::CreateProc(*AdbFilename, *Params, true, false, false, NULL, 0, NULL, LogcatWritePipe);
+	return LogcatProcHandle.IsValid();
+}
+
+inline bool FAndroidDeviceOutputReaderRunnable::Init(void) 
+{ 
+	FPlatformProcess::CreatePipe(LogcatReadPipe, LogcatWritePipe);
+	return StartLogcatProcess();
+}
+
+inline void FAndroidDeviceOutputReaderRunnable::Exit(void) 
+{
+	if (LogcatProcHandle.IsValid())
+	{
+		FPlatformProcess::CloseProc(LogcatProcHandle);
+	}
+	FPlatformProcess::ClosePipe(LogcatReadPipe, LogcatWritePipe);
+}
+
+inline void FAndroidDeviceOutputReaderRunnable::Stop(void)
+{
+	StopTaskCounter.Increment();
+}
+
+inline uint32 FAndroidDeviceOutputReaderRunnable::Run(void)
+{
+	FString LogcatOutput;
+	
+	while (StopTaskCounter.GetValue() == 0 && LogcatProcHandle.IsValid())
+	{
+		if (!FPlatformProcess::IsProcRunning(LogcatProcHandle))
+		{
+			// When user plugs out USB cable adb process stops
+			// Keep trying to restore adb connection until code that uses this object will not kill us
+			Output->Serialize(TEXT("Trying to restore connection to device..."), ELogVerbosity::Log, NAME_None);
+			FPlatformProcess::CloseProc(LogcatProcHandle);
+			if (StartLogcatProcess())
+			{
+				FPlatformProcess::Sleep(1.0f);
+			}
+			else
+			{
+				Output->Serialize(TEXT("Failed to start adb proccess"), ELogVerbosity::Log, NAME_None);
+				Stop();
+			}
+		}
+		else
+		{
+			LogcatOutput.Append(FPlatformProcess::ReadPipe(LogcatReadPipe));
+
+			if (LogcatOutput.Len() > 0)
+			{
+				TArray<FString> OutputLines;
+				LogcatOutput.ParseIntoArray(OutputLines, TEXT("\n"), false);
+
+				if (!LogcatOutput.EndsWith(TEXT("\n")))
+				{
+					// partial line at the end, do not serialize it until we receive remainder
+					LogcatOutput = OutputLines.Last();
+					OutputLines.RemoveAt(OutputLines.Num() - 1);
+				}
+				else
+				{
+					LogcatOutput.Reset();
+				}
+
+				for (int32 i = 0; i < OutputLines.Num(); ++i)
+				{
+					Output->Serialize(*OutputLines[i], ELogVerbosity::Log, NAME_None);
+				}
+			}
+			
+			FPlatformProcess::Sleep(0.1f);
+		}
+	}
+
+	return 0;
+}
+
+inline bool FAndroidTargetDeviceOutput::Init(const FAndroidTargetDevice& TargetDevice, FOutputDevice* Output)
+{
+	check(Output);
+	// Output will be produced by background thread
+	check(Output->CanBeUsedOnAnyThread());
+	
+	DeviceSerialNumber = TargetDevice.GetSerialNumber();
+	DeviceName = TargetDevice.GetName();
+		
+	FString AdbFilename;
+	if (FAndroidTargetDevice::GetAdbFullFilename(AdbFilename))
+	{
+		auto* Runnable = new FAndroidDeviceOutputReaderRunnable(AdbFilename, DeviceSerialNumber, Output);
+		DeviceOutputThread = TUniquePtr<FRunnableThread>(FRunnableThread::Create(Runnable, TEXT("FAndroidDeviceOutputReaderRunnable")));
+		return true;
+	}
+		
+	return false;
+}
+	
diff --git a/Engine/Source/Developer/Android/AndroidTargetPlatformControls/Private/AndroidTargetPlatformControls.cpp b/Engine/Source/Developer/Android/AndroidTargetPlatformControls/Private/AndroidTargetPlatformControls.cpp
new file mode 100644
--- /dev/null
+++ b/Engine/Source/Developer/Android/AndroidTargetPlatformControls/Private/AndroidTargetPlatformControls.cpp
@@ -0,0 +1,690 @@
+// Copyright Epic Games, Inc. All Rights Reserved.
+
+/*=============================================================================
+	AndroidTargetPlatform.inl: Implements the FAndroidTargetPlatformControls class.
+=============================================================================*/
+
+
+/* FAndroidTargetPlatformControls structors
+ *****************************************************************************/
+
+#include "AndroidTargetPlatformControls.h"
+#include "AndroidTargetPlatformSettings.h"
+#include "CoreTypes.h"
+#include "Misc/AssertionMacros.h"
+#include "Containers/Array.h"
+#include "Containers/UnrealString.h"
+#include "UObject/NameTypes.h"
+#include "Logging/LogMacros.h"
+#include "Stats/Stats.h"
+#include "Serialization/Archive.h"
+#include "Misc/FileHelper.h"
+#include "Misc/SecureHash.h"
+#include "HAL/FileManager.h"
+#include "HAL/PlatformFileManager.h"
+#include "HAL/IConsoleManager.h"
+#include "Interfaces/IAndroidDeviceDetectionModule.h"
+#include "Interfaces/IAndroidDeviceDetection.h"
+#include "Modules/ModuleManager.h"
+#include "Misc/SecureHash.h"
+#include "AnalyticsEventAttribute.h"
+
+
+#if WITH_ENGINE
+#include "AudioCompressionSettings.h"
+#include "Sound/SoundWave.h"
+#include "TextureResource.h"
+#endif
+
+#define LOCTEXT_NAMESPACE "FAndroidTargetPlatformControls"
+
+class Error;
+class FAndroidTargetDevice;
+class FConfigCacheIni;
+class FModuleManager;
+class FScopeLock;
+class FStaticMeshLODSettings;
+class FTargetDeviceId;
+class FTSTicker;
+class IAndroidDeviceDetectionModule;
+class UTexture;
+class UTextureLODSettings;
+struct FAndroidDeviceInfo;
+enum class ETargetPlatformFeatures;
+template<typename TPlatformProperties> class TTargetPlatformBase;
+
+static FString GetLicensePath()
+{
+	auto &AndroidDeviceDetection = FModuleManager::LoadModuleChecked<IAndroidDeviceDetectionModule>("AndroidDeviceDetection");
+	IAndroidDeviceDetection* DeviceDetection = AndroidDeviceDetection.GetAndroidDeviceDetection();
+	FString ADBPath = DeviceDetection->GetADBPath();
+
+	if (!FPaths::FileExists(ADBPath))
+	{
+		return TEXT("");
+	}
+
+	// strip off the adb.exe part
+	FString PlatformToolsPath;
+	FString Filename;
+	FString Extension;
+	FPaths::Split(ADBPath, PlatformToolsPath, Filename, Extension);
+
+	// remove the platform-tools part and point to licenses
+	FPaths::NormalizeDirectoryName(PlatformToolsPath);
+	FString LicensePath = PlatformToolsPath + "/../licenses";
+	FPaths::CollapseRelativeDirectories(LicensePath);
+
+	return LicensePath;
+}
+
+#if WITH_ENGINE
+static bool GetLicenseHash(FSHAHash& LicenseHash)
+{
+	bool bLicenseValid = false;
+
+	// from Android SDK Tools 25.2.3
+	FString LicenseFilename = FPaths::EngineDir() + TEXT("Source/ThirdParty/Android/package.xml");
+
+	// Create file reader
+	TUniquePtr<FArchive> FileReader(IFileManager::Get().CreateFileReader(*LicenseFilename));
+	if (FileReader)
+	{
+		// Create buffer for file input
+		uint32 BufferSize = IntCastChecked<uint32>(FileReader->TotalSize());
+		uint8* Buffer = (uint8*)FMemory::Malloc(BufferSize);
+		FileReader->Serialize(Buffer, BufferSize);
+
+		uint8 StartPattern[] = "<license id=\"android-sdk-license\" type=\"text\">";
+		int32 StartPatternLength = strlen((char *)StartPattern);
+
+		uint8* LicenseStart = Buffer;
+		uint8* BufferEnd = Buffer + BufferSize - StartPatternLength;
+		while (LicenseStart < BufferEnd)
+		{
+			if (!memcmp(LicenseStart, StartPattern, StartPatternLength))
+			{
+				break;
+			}
+			LicenseStart++;
+		}
+
+		if (LicenseStart < BufferEnd)
+		{
+			LicenseStart += StartPatternLength;
+
+			uint8 EndPattern[] = "</license>";
+			int32 EndPatternLength = strlen((char *)EndPattern);
+
+			uint8* LicenseEnd = LicenseStart;
+			BufferEnd = Buffer + BufferSize - EndPatternLength;
+			while (LicenseEnd < BufferEnd)
+			{
+				if (!memcmp(LicenseEnd, EndPattern, EndPatternLength))
+				{
+					break;
+				}
+				LicenseEnd++;
+			}
+
+			if (LicenseEnd < BufferEnd)
+			{
+				int32 LicenseLength = IntCastChecked<int32>(LicenseEnd - LicenseStart);
+				FSHA1::HashBuffer(LicenseStart, LicenseLength, LicenseHash.Hash);
+				bLicenseValid = true;
+			}
+		}
+		FMemory::Free(Buffer);
+	}
+
+	return bLicenseValid;
+}
+#endif
+
+static bool HasLicense()
+{
+#if WITH_ENGINE
+	FString LicensePath = GetLicensePath();
+
+	if (LicensePath.IsEmpty())
+	{
+		return false;
+	}
+
+	// directory must exist
+	IPlatformFile& PlatformFile = FPlatformFileManager::Get().GetPlatformFile();
+	if (!PlatformFile.DirectoryExists(*LicensePath))
+	{
+		return false;
+	}
+
+	// license file must exist
+	FString LicenseFilename = LicensePath + "/android-sdk-license";
+	if (!PlatformFile.FileExists(*LicenseFilename))
+	{
+		return false;
+	}
+
+	FSHAHash LicenseHash;
+	if (!GetLicenseHash(LicenseHash))
+	{
+		return false;
+	}
+
+	// contents must match hash of license text
+	FString FileData = "";
+	FFileHelper::LoadFileToString(FileData, *LicenseFilename);
+	TArray<FString> lines;
+	int32 lineCount = FileData.ParseIntoArray(lines, TEXT("\n"), true);
+
+	FString LicenseString = LicenseHash.ToString().ToLower();
+	for (FString &line : lines)
+	{
+		if (line.TrimStartAndEnd().Equals(LicenseString))
+		{
+			return true;
+		}
+	}
+#endif
+
+	// doesn't match
+	return false;
+}
+
+FAndroidTargetPlatformControls::FAndroidTargetPlatformControls(bool bInIsClient, ITargetPlatformSettings* InTargetPlatformSettings, const TCHAR* FlavorName, const TCHAR* OverrideIniPlatformName)
+	: TNonDesktopTargetPlatformControlsBase(bInIsClient, InTargetPlatformSettings, FlavorName, OverrideIniPlatformName)
+	, DeviceDetection(nullptr)
+
+{
+	TickDelegate = FTickerDelegate::CreateRaw(this, &FAndroidTargetPlatformControls::HandleTicker);
+	TickDelegateHandle = FTSTicker::GetCoreTicker().AddTicker(TickDelegate, 4.0f);
+
+	AndroidTargetPlatformSettings = (FAndroidTargetPlatformSettings*)TargetPlatformSettings;
+}
+
+
+FAndroidTargetPlatformControls::~FAndroidTargetPlatformControls()
+{
+	 FTSTicker::GetCoreTicker().RemoveTicker(TickDelegateHandle);
+}
+
+
+/* ITargetPlatform overrides
+ *****************************************************************************/
+
+void FAndroidTargetPlatformControls::GetAllDevices( TArray<ITargetDevicePtr>& OutDevices ) const
+{
+	OutDevices.Reset();
+
+	for (auto Iter = Devices.CreateConstIterator(); Iter; ++Iter)
+	{
+		OutDevices.Add(Iter.Value());
+	}
+}
+
+ITargetDevicePtr FAndroidTargetPlatformControls::GetDefaultDevice( ) const
+{
+	// return the first device in the list
+	if (Devices.Num() > 0)
+	{
+		auto Iter = Devices.CreateConstIterator();
+		if (Iter)
+		{
+			return Iter.Value();
+		}
+	}
+
+	return nullptr;
+}
+
+ITargetDevicePtr FAndroidTargetPlatformControls::GetDevice( const FTargetDeviceId& DeviceId )
+{
+	if (DeviceId.GetPlatformName() == PlatformName())
+	{
+		return Devices.FindRef(DeviceId.GetDeviceName());
+	}
+
+	return nullptr;
+}
+
+
+bool FAndroidTargetPlatformControls::IsSdkInstalled(bool bProjectHasCode, FString& OutDocumentationPath) const
+{
+	OutDocumentationPath = FString("Shared/Tutorials/SettingUpAndroidTutorial");
+	return true;
+}
+
+int32 FAndroidTargetPlatformControls::CheckRequirements(bool bProjectHasCode, EBuildConfiguration Configuration, bool bRequiresAssetNativization, FString& OutTutorialPath, FString& OutDocumentationPath, FText& CustomizedLogMessage) const
+{
+	OutDocumentationPath = TEXT("Platforms/Android/GettingStarted");
+
+	int32 bReadyToBuild = ETargetPlatformReadyStatus::Ready;
+	if (!IsSdkInstalled(bProjectHasCode, OutTutorialPath))
+	{
+		bReadyToBuild |= ETargetPlatformReadyStatus::SDKNotFound;
+	}
+
+	// need to check license was accepted
+	if (!HasLicense())
+	{
+		OutTutorialPath.Empty();
+		CustomizedLogMessage = LOCTEXT("AndroidLicenseNotAcceptedMessageDetail", "SDK License must be accepted in the Android project settings to deploy your app to the device.");
+		bReadyToBuild |= ETargetPlatformReadyStatus::LicenseNotAccepted;
+	}
+
+	return bReadyToBuild;
+}
+
+void FAndroidTargetPlatformControls::GetPlatformSpecificProjectAnalytics(TArray<FAnalyticsEventAttribute>& AnalyticsParamArray) const
+{
+	TTargetPlatformControlsBase<FAndroidPlatformProperties>::GetPlatformSpecificProjectAnalytics(AnalyticsParamArray);
+
+	AppendAnalyticsEventAttributeArray(AnalyticsParamArray,
+		TEXT("AndroidVariant"), GetAndroidVariantName(),
+		TEXT("SupportsVulkan"), AndroidTargetPlatformSettings->SupportsVulkan(),
+		TEXT("SupportsVulkanSM5"), AndroidTargetPlatformSettings->SupportsVulkanSM5(),
+		TEXT("SupportsES31"), AndroidTargetPlatformSettings->SupportsES31()
+	);
+
+	AppendAnalyticsEventConfigBool(AnalyticsParamArray, TEXT("/Script/AndroidRuntimeSettings.AndroidRuntimeSettings"), TEXT("bPackageForMetaQuest"), GEngineIni);
+}
+
+#if WITH_ENGINE
+
+FName FAndroidTargetPlatformControls::FinalizeVirtualTextureLayerFormat(FName Format) const
+{
+#if WITH_EDITOR
+	// Remap non-ETC variants to ETC
+
+	// VirtualTexture Format was already run through the ordinary texture remaps to change AutoDXT to ASTC or ETC
+	// this then runs again
+	// currently it forces all ASTC to ETC
+	// this is needed because the runtime virtual texture encoder only supports ETC
+
+	// code dupe with IOSTargetPlatform
+
+	const static FName VTRemap[][2] =
+	{
+		{ { FName(TEXT("ASTC_RGB")) },			{ AndroidTexFormat::NameETC2_RGB } },
+		{ { FName(TEXT("ASTC_RGBA")) },			{ AndroidTexFormat::NameETC2_RGBA } },
+		{ { FName(TEXT("ASTC_RGBA_HQ")) },		{ AndroidTexFormat::NameETC2_RGBA } },
+//		{ { FName(TEXT("ASTC_RGB_HDR")) },		{ NameRGBA16F } }, // ?
+		{ { FName(TEXT("ASTC_RGBAuto")) },		{ AndroidTexFormat::NameAutoETC2 } },
+		{ { FName(TEXT("ASTC_NormalAG")) },		{ AndroidTexFormat::NameETC2_RGB } },
+		{ { AndroidTexFormat::NameASTC_NormalRG },	{ AndroidTexFormat::NameETC2_RG11 } },
+		{ { AndroidTexFormat::NameASTC_NormalLA	},	{ AndroidTexFormat::NameETC2_RG11 } },
+		{ { AndroidTexFormat::NameDXT1 },		{ AndroidTexFormat::NameETC2_RGB } },
+		{ { AndroidTexFormat::NameDXT5 },		{ AndroidTexFormat::NameETC2_RGBA } },
+		{ { AndroidTexFormat::NameAutoDXT },	{ AndroidTexFormat::NameAutoETC2 } }
+	};
+
+	for (int32 RemapIndex = 0; RemapIndex < UE_ARRAY_COUNT(VTRemap); RemapIndex++)
+	{
+		if (VTRemap[RemapIndex][0] == Format)
+		{
+			return VTRemap[RemapIndex][1];
+		}
+	}
+#endif
+	return Format;
+}
+#endif //WITH_ENGINE
+
+bool FAndroidTargetPlatformControls::SupportsVariants() const
+{
+	return true;
+}
+
+
+/* FAndroidTargetPlatformControls implementation
+ *****************************************************************************/
+void FAndroidTargetPlatformControls::InitializeDeviceDetection()
+{
+	DeviceDetection = FModuleManager::LoadModuleChecked<IAndroidDeviceDetectionModule>("AndroidDeviceDetection").GetAndroidDeviceDetection();
+	DeviceDetection->Initialize(TEXT("ANDROID_HOME"),
+#if PLATFORM_WINDOWS
+		TEXT("platform-tools\\adb.exe"),
+#else
+		TEXT("platform-tools/adb"),
+#endif
+		TEXT("shell getprop"), true);
+}
+
+bool FAndroidTargetPlatformControls::ShouldExpandTo32Bit(const uint16* Indices, const int32 NumIndices) const
+{
+	bool bIsMaliBugIndex = false;
+	const uint16 MaliBugIndexMaxDiff = 16;
+	uint16 LastIndex = Indices[0];
+	for (int32 i = 1; i < NumIndices; ++i)
+	{
+		uint16 CurrentIndex = Indices[i];
+		if ((FMath::Abs(LastIndex - CurrentIndex) > MaliBugIndexMaxDiff))
+		{
+			bIsMaliBugIndex = true;
+			break;
+		}
+		else
+		{
+			LastIndex = CurrentIndex;
+		}
+	}
+	return bIsMaliBugIndex;
+}
+
+/* FAndroidTargetPlatformControls callbacks
+ *****************************************************************************/
+
+bool FAndroidTargetPlatformControls::HandleTicker( float DeltaTime )
+{
+	QUICK_SCOPE_CYCLE_COUNTER(STAT_FAndroidTargetPlatform_HandleTicker);
+
+	if (DeviceDetection == nullptr)
+	{
+		InitializeDeviceDetection();
+		checkf(DeviceDetection != nullptr, TEXT("A target platform didn't create a device detection object in InitializeDeviceDetection()!"));
+	}
+
+	TArray<FString> ConnectedDeviceIds;
+
+	{
+		FScopeLock ScopeLock(DeviceDetection->GetDeviceMapLock());
+
+		auto DeviceIt = DeviceDetection->GetDeviceMap().CreateConstIterator();
+
+		for (; DeviceIt; ++DeviceIt)
+		{
+			ConnectedDeviceIds.Add(DeviceIt.Key());
+
+			const FAndroidDeviceInfo& DeviceInfo = DeviceIt.Value();
+
+			// see if this device is already known
+			if (Devices.Contains(DeviceIt.Key()))
+			{
+				FAndroidTargetDevicePtr TestDevice = Devices[DeviceIt.Key()];
+
+				// ignore if authorization didn't change
+				if (DeviceInfo.bAuthorizedDevice == TestDevice->IsAuthorized())
+				{
+					continue;
+				}
+
+				// remove it to add again
+				TestDevice->SetConnected(false);
+				Devices.Remove(DeviceIt.Key());
+
+				OnDeviceLost().Broadcast(TestDevice.ToSharedRef());
+			}
+
+			// check if this platform is supported by the extensions and version
+			if (!SupportedByExtensionsString(DeviceInfo.GLESExtensions, DeviceInfo.GLESVersion))
+			{
+				continue;
+			}
+
+			// create target device
+			FAndroidTargetDevicePtr& Device = Devices.Add(DeviceInfo.SerialNumber);
+
+			Device = CreateTargetDevice(*this, DeviceInfo.SerialNumber, GetAndroidVariantName());
+
+			Device->SetConnected(true);
+			Device->SetModel(DeviceInfo.Model);
+			Device->SetDeviceName(DeviceInfo.DeviceName);
+			Device->SetAuthorized(DeviceInfo.bAuthorizedDevice);
+			Device->SetVersions(DeviceInfo.SDKVersion, DeviceInfo.HumanAndroidVersion);
+
+			OnDeviceDiscovered().Broadcast(Device.ToSharedRef());
+		}
+	}
+
+	// remove disconnected devices
+	for (auto Iter = Devices.CreateIterator(); Iter; ++Iter)
+	{
+		if (!ConnectedDeviceIds.Contains(Iter.Key()))
+		{
+			FAndroidTargetDevicePtr RemovedDevice = Iter.Value();
+			RemovedDevice->SetConnected(false);
+
+			Iter.RemoveCurrent();
+
+			OnDeviceLost().Broadcast(RemovedDevice.ToSharedRef());
+		}
+	}
+
+	return true;
+}
+
+FAndroidTargetDeviceRef FAndroidTargetPlatformControls::CreateNewDevice(const FAndroidDeviceInfo &DeviceInfo)
+{
+	return MakeShareable(new FAndroidTargetDevice(*this, DeviceInfo.SerialNumber, GetAndroidVariantName()));
+}
+
+FAndroidTargetDevicePtr FAndroidTargetPlatformControls::CreateTargetDevice(const ITargetPlatformControls& InTargetPlatform, const FString& InSerialNumber, const FString& InAndroidVariant) const
+{
+	return MakeShareable(new FAndroidTargetDevice(InTargetPlatform, InSerialNumber, InAndroidVariant));
+}
+
+#if WITH_ENGINE
+void FAndroidTargetPlatformControls::GetTextureFormats(const UTexture* Texture, TArray< TArray<FName> >& OutFormats) const
+{
+	// FAndroidTargetPlatformControls aside from being the base class for all the concrete android target platforms
+	//	it is also usable on its own as "flavorless" Android
+	//	but I don't understand how that's supposed to work or what that's supposed to mean
+	//	and no information has been forthcoming
+	check(Texture);
+
+	// Supported in ES3.2 with ASTC
+	const bool bSupportCompressedVolumeTexture = AndroidTargetPlatformSettings->SupportsTextureFormatCategory(EAndroidTextureFormatCategory::ASTC);
+	// OpenGL ES has F32 textures but doesn't allow linear filtering unless OES_texture_float_linear
+	const bool bSupportFilteredFloat32Textures = false;
+
+	// optionaly compress landscape weightmaps for a mobile rendering
+	static const auto CompressLandscapeWeightMapsVar = IConsoleManager::Get().FindTConsoleVariableDataInt(TEXT("r.Mobile.CompressLandscapeWeightMaps"));
+	static const bool bCompressLandscapeWeightMaps = (CompressLandscapeWeightMapsVar && CompressLandscapeWeightMapsVar->GetValueOnAnyThread() != 0);
+
+	TArray<FName>& LayerFormats = OutFormats.AddDefaulted_GetRef();
+	int32 BlockSize = 1; // this looks wrong? should be 4 for FAndroid_DXTTargetPlatform ? - it is wrong, but BlockSize is ignored
+
+	GetDefaultTextureFormatNamePerLayer(LayerFormats, this->GetTargetPlatformSettings(), this, Texture, bSupportCompressedVolumeTexture, BlockSize, bSupportFilteredFloat32Textures);
+
+	for (FName& TextureFormatName : LayerFormats)
+	{
+		// @todo Oodle: this should not be here
+		//	should be in GetDefaultTextureFormatNamePerLayer
+		//	so that 4x4 checks can be applied correctly, etc.
+		if (Texture->LODGroup == TEXTUREGROUP_Terrain_Weightmap && bCompressLandscapeWeightMaps)
+		{
+			TextureFormatName = AndroidTexFormat::NameAutoDXT;
+		}
+
+		if (Texture->GetTextureClass() == ETextureClass::Cube)
+		{
+			FTextureFormatSettings FormatSettings;
+			Texture->GetDefaultFormatSettings(FormatSettings);
+			// TC_EncodedReflectionCapture is no longer used and could be deleted
+			if (FormatSettings.CompressionSettings == TC_EncodedReflectionCapture && !FormatSettings.CompressionNone)
+			{
+				TextureFormatName = FName(TEXT("ETC2_RGBA"));
+			}
+		}
+
+		for (int32 RemapIndex = 0; RemapIndex < UE_ARRAY_COUNT(AndroidTexFormat::GenericRemap); ++RemapIndex)
+		{
+			if (TextureFormatName == AndroidTexFormat::GenericRemap[RemapIndex][0])
+			{
+				TextureFormatName = AndroidTexFormat::GenericRemap[RemapIndex][1];
+			}
+		}
+	}
+}
+
+void FAndroidTargetPlatformControls::GetAllTextureFormats(TArray<FName>& OutFormats) const
+{
+	GetAllDefaultTextureFormats(this->GetTargetPlatformSettings(), OutFormats);
+
+	for (int32 RemapIndex = 0; RemapIndex < UE_ARRAY_COUNT(AndroidTexFormat::GenericRemap); ++RemapIndex)
+	{
+		OutFormats.Remove(AndroidTexFormat::GenericRemap[RemapIndex][0]);
+	}
+
+	for (int32 RemapIndex = 0; RemapIndex < UE_ARRAY_COUNT(AndroidTexFormat::GenericRemap); ++RemapIndex)
+	{
+		OutFormats.AddUnique(AndroidTexFormat::GenericRemap[RemapIndex][1]);
+	}
+}
+
+void FAndroid_ASTCTargetPlatformControls::GetAllTextureFormats(TArray<FName>& OutFormats) const
+{
+	FAndroidTargetPlatformControls::GetAllTextureFormats(OutFormats);
+
+	for (int32 RemapIndex = 0; RemapIndex < UE_ARRAY_COUNT(AndroidTexFormat::ASTCRemap); ++RemapIndex)
+	{
+		OutFormats.Remove(AndroidTexFormat::ASTCRemap[RemapIndex][0]);
+	}
+
+	// ASTC for compressed textures
+	OutFormats.Add(AndroidTexFormat::NameAutoASTC);
+	// ETC for ETC2_R11
+	OutFormats.Add(AndroidTexFormat::NameAutoETC2);
+}
+
+void FAndroid_ASTCTargetPlatformControls::GetTextureFormats(const UTexture* Texture, TArray< TArray<FName> >& OutFormats) const
+{
+	FAndroidTargetPlatformControls::GetTextureFormats(Texture, OutFormats);
+
+	// L+A mode for normal map compression
+	const bool bSupportsNormalLA = GetTargetPlatformSettings()->SupportsFeature(ETargetPlatformFeatures::NormalmapLAEncodingMode);
+
+	// perform any remapping away from defaults
+	TArray<FName>& LayerFormats = OutFormats.Last();
+	for (FName& TextureFormatName : LayerFormats)
+	{
+		if (bSupportsNormalLA && TextureFormatName == AndroidTexFormat::NameBC5)
+		{
+			TextureFormatName = AndroidTexFormat::NameASTC_NormalLA;
+			continue;
+		}
+
+		for (int32 RemapIndex = 0; RemapIndex < UE_ARRAY_COUNT(AndroidTexFormat::ASTCRemap); ++RemapIndex)
+		{
+			if (TextureFormatName == AndroidTexFormat::ASTCRemap[RemapIndex][0])
+			{
+				TextureFormatName = AndroidTexFormat::ASTCRemap[RemapIndex][1];
+				break;
+			}
+		}
+	}
+
+	bool bSupportASTCHDR = AndroidTargetPlatformSettings->UsesASTCHDR();
+
+	if (!bSupportASTCHDR)
+	{
+		for (FName& TextureFormatName : LayerFormats)
+		{
+			if (TextureFormatName == AndroidTexFormat::NameASTC_RGB_HDR)
+			{
+				TextureFormatName = AndroidTexFormat::NameRGBA16F;
+			}
+		}
+	}
+}
+
+void FAndroid_DXTTargetPlatformControls::GetTextureFormats(const UTexture* Texture, TArray< TArray<FName> >& OutFormats) const
+{
+	FAndroidTargetPlatformControls::GetTextureFormats(Texture, OutFormats);
+
+	bool bSupportsDX11Formats = false; // assume Android DXT does not support BC6/7
+
+	if (!bSupportsDX11Formats)
+	{
+		TArray<FName>& LayerFormats = OutFormats.Last();
+
+		for (FName& TextureFormatName : LayerFormats)
+		{
+			if (TextureFormatName == AndroidTexFormat::NameBC6H)
+			{
+				TextureFormatName = AndroidTexFormat::NameRGBA16F;
+			}
+			else if (TextureFormatName == AndroidTexFormat::NameBC7)
+			{
+				TextureFormatName = AndroidTexFormat::NameDXT5;
+			}
+		}
+	}
+}
+
+void FAndroid_ETC2TargetPlatformControls::GetAllTextureFormats(TArray<FName>& OutFormats) const
+{
+	FAndroidTargetPlatformControls::GetAllTextureFormats(OutFormats);
+
+	for (int32 RemapIndex = 0; RemapIndex < UE_ARRAY_COUNT(AndroidTexFormat::ETCRemap); ++RemapIndex)
+	{
+		OutFormats.Remove(AndroidTexFormat::ETCRemap[RemapIndex][0]);
+	}
+
+	// support only ETC for compressed textures
+	OutFormats.Add(AndroidTexFormat::NameAutoETC2);
+}
+
+void FAndroid_ETC2TargetPlatformControls::GetTextureFormats(const UTexture* Texture, TArray< TArray<FName> >& OutFormats) const
+{
+	FAndroidTargetPlatformControls::GetTextureFormats(Texture, OutFormats);
+
+	// perform any remapping away from defaults
+	TArray<FName>& LayerFormats = OutFormats.Last();
+	for (FName& TextureFormatName : LayerFormats)
+	{
+		for (int32 RemapIndex = 0; RemapIndex < UE_ARRAY_COUNT(AndroidTexFormat::ETCRemap); ++RemapIndex)
+		{
+			if (TextureFormatName == AndroidTexFormat::ETCRemap[RemapIndex][0])
+			{
+				TextureFormatName = AndroidTexFormat::ETCRemap[RemapIndex][1];
+				break;
+			}
+		}
+	}
+}
+
+void FAndroid_MultiTargetPlatformControls::GetTextureFormats(const UTexture* Texture, TArray< TArray<FName> >& OutFormats) const
+{
+	// Ask each platform variant to choose texture formats
+	for (ITargetPlatformControls* Platform : FormatTargetPlatforms)
+	{
+		TArray< TArray<FName> > PlatformFormats;
+		Platform->GetTextureFormats(Texture, PlatformFormats);
+		for (TArray<FName>& FormatPerLayer : PlatformFormats)
+		{
+			// For multiformat case we have to disable L+A normal map compression as only ASTC textures support it 
+			for (FName& TextureFormatName : FormatPerLayer)
+			{
+				if (TextureFormatName == AndroidTexFormat::NameASTC_NormalLA)
+				{
+					TextureFormatName = AndroidTexFormat::NameASTC_NormalRG;
+				}
+			}
+
+			OutFormats.AddUnique(FormatPerLayer);
+		}
+	}
+}
+
+void FAndroid_MultiTargetPlatformControls::GetAllTextureFormats(TArray<FName>& OutFormats) const
+{
+	// Ask each platform variant to choose texture formats
+	for (ITargetPlatformControls* Platform : FormatTargetPlatforms)
+	{
+		TArray<FName> PlatformFormats;
+		Platform->GetAllTextureFormats(PlatformFormats);
+		for (FName Format : PlatformFormats)
+		{
+			OutFormats.AddUnique(Format);
+		}
+	}
+}
+#endif
+
+FText FAndroid_MultiTargetPlatformControls::DisplayName() const
+{
+	return FText::Format(LOCTEXT("Android_Multi", "Android (Multi:{0})"), FText::FromString(FormatTargetString));
+}
+
+#undef LOCTEXT_NAMESPACE
diff --git a/Engine/Source/Developer/Android/AndroidTargetPlatformControls/Private/AndroidTargetPlatformControls.h b/Engine/Source/Developer/Android/AndroidTargetPlatformControls/Private/AndroidTargetPlatformControls.h
new file mode 100644
--- /dev/null
+++ b/Engine/Source/Developer/Android/AndroidTargetPlatformControls/Private/AndroidTargetPlatformControls.h
@@ -0,0 +1,312 @@
+// Copyright Epic Games, Inc. All Rights Reserved.
+
+#pragma once
+
+#include "CoreTypes.h"
+#include "Containers/UnrealString.h"
+#include "UObject/NameTypes.h"
+#include "Delegates/IDelegateInstance.h"
+#include "Containers/Map.h"
+#include "Delegates/Delegate.h"
+#include "Containers/Ticker.h"
+#include "Misc/ScopeLock.h"
+#include "Android/AndroidPlatformProperties.h"
+#include "Interfaces/ITargetPlatformModule.h"
+#include "Common/TargetPlatformControlsBase.h"
+#include "Interfaces/IAndroidDeviceDetection.h"
+#include "Interfaces/IAndroidDeviceDetectionModule.h"
+#include "AndroidTargetDevice.h"
+#include "AndroidTargetPlatformSettings.h"
+
+#if WITH_ENGINE
+#include "Engine/Texture.h"
+#include "Internationalization/Text.h"
+#include "StaticMeshResources.h"
+#endif // WITH_ENGINE
+
+
+class FTargetDeviceId;
+class IAndroidDeviceDetection;
+class ITargetPlatform;
+class UTextureLODSettings;
+enum class ETargetPlatformFeatures;
+template<typename TPlatformProperties> class TTargetPlatformBase;
+
+template< typename InElementType, typename KeyFuncs, typename Allocator > class TSet;
+template<typename KeyType, typename ValueType, typename SetAllocator, typename KeyFuncs > class TMap;
+template<typename KeyType, typename ValueType, typename SetAllocator, typename KeyFuncs > class TMultiMap;
+template<typename TPlatformProperties> class TTargetPlatformBase;
+
+/**
+ * FAndroidTargetPlatformControls, abstraction for cooking Android platforms
+ */
+class ANDROIDTARGETPLATFORMCONTROLS_API FAndroidTargetPlatformControls
+	: public TNonDesktopTargetPlatformControlsBase<FAndroidPlatformProperties>
+{
+public:
+
+	/**
+	 * Default constructor.
+	 */
+	FAndroidTargetPlatformControls(bool bInIsClient, ITargetPlatformSettings* TargetPlatformSettings, const TCHAR* FlavorName, const TCHAR* OverrideIniPlatformName = nullptr);
+
+	/**
+	 * Destructor
+	 */
+	virtual ~FAndroidTargetPlatformControls();
+
+public:
+
+	/**
+	 * Gets the name of the Android platform variant, i.e. ASTC, ETC2, DXT, etc.
+	 *
+	 * @param Variant name.
+	 */
+	FString GetAndroidVariantName() const
+	{
+		return PlatformInfo->PlatformFlavor.ToString();
+	}
+
+public:
+
+	//~ Begin ITargetPlatform Interface
+
+	virtual void EnableDeviceCheck(bool OnOff) override {}
+
+	virtual bool AddDevice(const FString& DeviceName, bool bDefault) override
+	{
+		return false;
+	}
+
+	virtual void GetAllDevices(TArray<ITargetDevicePtr>& OutDevices) const override;
+
+	virtual bool GenerateStreamingInstallManifest(const TMultiMap<FString, int32>& PakchunkMap, const TSet<int32>& PakchunkIndicesInUse) const override
+	{
+		return true;
+	}
+
+	virtual ITargetDevicePtr GetDefaultDevice() const override;
+
+	virtual ITargetDevicePtr GetDevice(const FTargetDeviceId& DeviceId) override;
+
+	virtual bool IsSdkInstalled(bool bProjectHasCode, FString& OutDocumentationPath) const override;
+
+	virtual int32 CheckRequirements(bool bProjectHasCode, EBuildConfiguration Configuration, bool bRequiresAssetNativization, FString& OutTutorialPath, FString& OutDocumentationPath, FText& CustomizedLogMessage) const override;
+
+	virtual void GetPlatformSpecificProjectAnalytics(TArray<struct FAnalyticsEventAttribute>& AnalyticsParamArray) const override;
+
+	virtual bool SupportsCompressedNonPOT() const
+	{
+		// most formats do support non-POT compressed textures
+		return true;
+	}
+
+#if WITH_ENGINE
+	virtual void GetTextureFormats(const UTexture* InTexture, TArray< TArray<FName> >& OutFormats) const override;
+	virtual FName FinalizeVirtualTextureLayerFormat(FName Format) const override;
+	virtual void GetAllTextureFormats(TArray<FName>& OutFormats) const override;
+#endif //WITH_ENGINE
+
+	virtual bool SupportsVariants() const override;
+
+	virtual void GetBuildProjectSettingKeys(FString& OutSection, TArray<FString>& InBoolKeys, TArray<FString>& InIntKeys, TArray<FString>& InStringKeys) const override
+	{
+		OutSection = TEXT("/Script/AndroidRuntimeSettings.AndroidRuntimeSettings");
+		InBoolKeys.Add(TEXT("bBuildForArm64"));	InBoolKeys.Add(TEXT("bBuildForX8664"));
+		InBoolKeys.Add(TEXT("bBuildForES31")); InBoolKeys.Add(TEXT("bBuildWithHiddenSymbolVisibility"));
+		InBoolKeys.Add(TEXT("bSaveSymbols")); InStringKeys.Add(TEXT("NDKAPILevel"));
+	}
+
+	virtual bool ShouldExpandTo32Bit(const uint16* Indices, const int32 NumIndices) const override;
+	//~ End ITargetPlatform Interface
+
+	virtual void InitializeDeviceDetection();
+
+protected:
+
+	virtual FAndroidTargetDevicePtr CreateTargetDevice(const ITargetPlatformControls& InTargetPlatform, const FString& InSerialNumber, const FString& InAndroidVariant) const;
+
+protected:
+	/**
+	 * Return true if this device has a supported set of extensions for this platform.
+	 *
+	 * @param Extensions - The GL extensions string.
+	 * @param GLESVersion - The GLES version reported by this device.
+	 */
+	virtual bool SupportedByExtensionsString(const FString& ExtensionsString, const int GLESVersion) const
+	{
+		return true;
+	}
+
+	// Handles when the ticker fires.
+	bool HandleTicker(float DeltaTime);
+
+	virtual FAndroidTargetDeviceRef CreateNewDevice(const FAndroidDeviceInfo& DeviceInfo);
+
+	// Holds a map of valid devices.
+	TMap<FString, FAndroidTargetDevicePtr> Devices;
+
+	// Holds a delegate to be invoked when the widget ticks.
+	FTickerDelegate TickDelegate;
+
+	// Handle to the registered TickDelegate.
+	FTSTicker::FDelegateHandle TickDelegateHandle;
+
+	// Pointer to the device detection handler that grabs device ids in another thread
+	IAndroidDeviceDetection* DeviceDetection;
+
+#if WITH_ENGINE
+
+	ITargetDevicePtr DefaultDevice;
+#endif //WITH_ENGINE
+
+	FAndroidTargetPlatformSettings* AndroidTargetPlatformSettings;
+};
+
+class FAndroid_DXTTargetPlatformControls : public FAndroidTargetPlatformControls
+{
+public:
+	FAndroid_DXTTargetPlatformControls(bool bIsClient, ITargetPlatformSettings* TargetPlatformSettings) : FAndroidTargetPlatformControls(bIsClient, TargetPlatformSettings, TEXT("DXT"))
+	{
+	}
+
+	virtual bool SupportedByExtensionsString(const FString& ExtensionsString, const int GLESVersion) const override
+	{
+		return (ExtensionsString.Contains(TEXT("GL_NV_texture_compression_s3tc")) || ExtensionsString.Contains(TEXT("GL_EXT_texture_compression_s3tc")));
+	}
+
+	virtual float GetVariantPriority() const override
+	{
+		float Priority;
+		return (GConfig->GetFloat(TEXT("/Script/AndroidRuntimeSettings.AndroidRuntimeSettings"), TEXT("TextureFormatPriority_DXT"), Priority, GEngineIni) ?
+			Priority : 0.6f) * 10.0f + (IsClientOnly() ? 0.25f : 0.5f);
+	}
+
+#if WITH_ENGINE
+	virtual void GetTextureFormats(const UTexture* Texture, TArray< TArray<FName> >& OutFormats) const override;
+#endif
+};
+
+class FAndroid_ASTCTargetPlatformControls : public FAndroidTargetPlatformControls
+{
+public:
+	FAndroid_ASTCTargetPlatformControls(bool bIsClient, ITargetPlatformSettings* TargetPlatformSettings) : FAndroidTargetPlatformControls(bIsClient, TargetPlatformSettings, TEXT("ASTC"))
+	{
+	}
+
+	virtual bool SupportedByExtensionsString(const FString& ExtensionsString, const int GLESVersion) const override
+	{
+		return ExtensionsString.Contains(TEXT("GL_KHR_texture_compression_astc_ldr"));
+	}
+
+	virtual float GetVariantPriority() const override
+	{
+		float Priority;
+		return (GConfig->GetFloat(TEXT("/Script/AndroidRuntimeSettings.AndroidRuntimeSettings"), TEXT("TextureFormatPriority_ASTC"), Priority, GEngineIni) ?
+			Priority : 0.9f) * 10.0f + (IsClientOnly() ? 0.25f : 0.5f);
+	}
+
+#if WITH_ENGINE
+	virtual void GetTextureFormats(const UTexture* Texture, TArray< TArray<FName> >& OutFormats) const;
+	virtual void GetAllTextureFormats(TArray<FName>& OutFormats) const override;
+#endif
+};
+
+
+class FAndroid_ETC2TargetPlatformControls : public FAndroidTargetPlatformControls
+{
+public:
+
+	FAndroid_ETC2TargetPlatformControls(bool bIsClient, ITargetPlatformSettings* TargetPlatformSettings) : FAndroidTargetPlatformControls(bIsClient, TargetPlatformSettings, TEXT("ETC2"))
+	{
+	}
+
+	virtual bool SupportedByExtensionsString(const FString& ExtensionsString, const int GLESVersion) const override
+	{
+		return GLESVersion >= 0x30000;
+	}
+
+	virtual float GetVariantPriority() const override
+	{
+		float Priority;
+		return (GConfig->GetFloat(TEXT("/Script/AndroidRuntimeSettings.AndroidRuntimeSettings"), TEXT("TextureFormatPriority_ETC2"), Priority, GEngineIni) ?
+			Priority : 0.2f) * 10.0f + (IsClientOnly() ? 0.25f : 0.5f);
+	}
+
+#if WITH_ENGINE
+	virtual void GetTextureFormats(const UTexture* Texture, TArray< TArray<FName> >& OutFormats) const;
+	virtual void GetAllTextureFormats(TArray<FName>& OutFormats) const override;
+#endif
+};
+
+class FAndroid_MultiTargetPlatformControls : public FAndroidTargetPlatformControls
+{
+	TArray<FAndroidTargetPlatformControls*> FormatTargetPlatforms;
+	FString FormatTargetString;
+
+public:
+	FAndroid_MultiTargetPlatformControls(bool bIsClient, ITargetPlatformSettings* TargetPlatformSettings) : FAndroidTargetPlatformControls(bIsClient, TargetPlatformSettings, TEXT("Multi"))
+	{
+	}
+
+	// set up all of the multiple formats together into this one
+	void LoadFormats(TArray<FAndroidTargetPlatformControls*> SingleFormatTPs)
+	{
+		// sort formats by priority so higher priority formats are packaged (and thus used by the device) first
+		// note that we passed this by value, not ref, so we can sort it
+		SingleFormatTPs.Sort([](const FAndroidTargetPlatformControls& A, const FAndroidTargetPlatformControls& B)
+			{
+				float PriorityA = 0.f;
+				float PriorityB = 0.f;
+				FString VariantA = A.GetAndroidVariantName().Replace(TEXT("Client"), TEXT(""));
+				FString VariantB = B.GetAndroidVariantName().Replace(TEXT("Client"), TEXT(""));
+				GConfig->GetFloat(TEXT("/Script/AndroidRuntimeSettings.AndroidRuntimeSettings"), *(FString(TEXT("TextureFormatPriority_")) + VariantA), PriorityA, GEngineIni);
+				GConfig->GetFloat(TEXT("/Script/AndroidRuntimeSettings.AndroidRuntimeSettings"), *(FString(TEXT("TextureFormatPriority_")) + VariantB), PriorityB, GEngineIni);
+				return PriorityA > PriorityB;
+			});
+
+		FormatTargetPlatforms.Empty();
+		FormatTargetString = TEXT("");
+
+		TSet<FString> SeenFormats;
+
+		// Load the TargetPlatform module for each format
+		for (FAndroidTargetPlatformControls* SingleFormatTP : SingleFormatTPs)
+		{
+			// only use once each
+			if (SeenFormats.Contains(SingleFormatTP->GetAndroidVariantName()))
+			{
+				continue;
+			}
+			SeenFormats.Add(SingleFormatTP->GetAndroidVariantName());
+
+			bool bEnabled = false;
+			FString SettingsName = FString(TEXT("bMultiTargetFormat_")) + *SingleFormatTP->GetAndroidVariantName();
+			GConfig->GetBool(TEXT("/Script/AndroidRuntimeSettings.AndroidRuntimeSettings"), *SettingsName, bEnabled, GEngineIni);
+			if (bEnabled)
+			{
+				if (FormatTargetPlatforms.Num())
+				{
+					FormatTargetString += TEXT(",");
+				}
+				FormatTargetString += SingleFormatTP->GetAndroidVariantName();
+				FormatTargetPlatforms.Add(SingleFormatTP);
+			}
+		}
+
+		// @todo do we need this with DisplayName below?
+		PlatformInfo::UpdatePlatformDisplayName(TEXT("Android_Multi"), DisplayName());
+	}
+
+	virtual float GetVariantPriority() const override
+	{
+		// lowest priority so specific variants are chosen first
+		return (IsClientOnly() ? 0.25f : 0.5f);
+	}
+
+#if WITH_ENGINE
+	virtual void GetTextureFormats(const UTexture* Texture, TArray< TArray<FName> >& OutFormats) const;
+	virtual void GetAllTextureFormats(TArray<FName>& OutFormats) const override;
+#endif	
+	virtual FText DisplayName() const override;
+};
diff --git a/Engine/Source/Developer/Android/AndroidTargetPlatformControls/Private/AndroidTargetPlatformControlsModule.cpp b/Engine/Source/Developer/Android/AndroidTargetPlatformControls/Private/AndroidTargetPlatformControlsModule.cpp
new file mode 100644
--- /dev/null
+++ b/Engine/Source/Developer/Android/AndroidTargetPlatformControls/Private/AndroidTargetPlatformControlsModule.cpp
@@ -0,0 +1,91 @@
+// Copyright Epic Games, Inc. All Rights Reserved.
+
+#include "CoreMinimal.h"
+#include "AndroidTargetPlatformControls.h"
+#include "Modules/ModuleManager.h"
+#include "Interfaces/ITargetPlatformModule.h"
+#include "Interfaces/ITargetPlatformControlsModule.h"
+#include "Android/AndroidPlatformProperties.h"
+#include "MultiPlatformTargetReceiptBuildWorkers.h"
+#include "IAndroidTargetPlatformControlsModule.h"
+#include "IAndroidTargetPlatformSettingsModule.h"
+
+#define LOCTEXT_NAMESPACE "FAndroidTargetPlatformControlsModule"
+
+/**
+ * Module for the Android target platform controls.
+ */
+class FAndroidTargetPlatformControlsModule
+	: public IAndroidTargetPlatformControlsModule
+{
+public:
+	virtual void GetTargetPlatforms(TArray<ITargetPlatform*>& TargetPlatforms)
+	{
+
+	}
+	/** Destructor. */
+	~FAndroidTargetPlatformControlsModule()
+	{
+	}
+protected:
+	virtual void GetTargetPlatformControls(TArray<ITargetPlatformControls*>& TargetPlatforms, FName& PlatformSettingsModuleName)
+	{
+		if (FAndroidTargetPlatformControls::IsUsable())
+		{
+			// ensure this is only called once
+			check(SinglePlatforms.IsEmpty());
+
+			TMap<EAndroidTextureFormatCategory, ITargetPlatformSettings*> OutMap;
+			ITargetPlatformSettings* OutMulti = nullptr;
+			IAndroidTargetPlatformSettingsModule* Module = FModuleManager::GetModulePtr<IAndroidTargetPlatformSettingsModule>(PlatformSettingsModuleName);
+			if (Module != nullptr)
+			{
+				Module->GetSinglePlatformSettings(OutMap);
+				OutMulti = Module->GetMultiPlatformSettings();
+			}
+
+			for (int32 Type = 0; Type < 2; Type++)
+			{
+				bool bIsClient = Type == 1;
+
+				// flavorless must come first
+				// "flavorless" (aka "vanilla") Android is not well defined and not actually usable
+				//  but Editor wants to see it in the platform list
+				SinglePlatforms.Add(new FAndroidTargetPlatformControls(bIsClient, OutMap[EAndroidTextureFormatCategory::Count], nullptr));
+				SinglePlatforms.Add(new FAndroid_ASTCTargetPlatformControls(bIsClient, OutMap[EAndroidTextureFormatCategory::ASTC]));
+				SinglePlatforms.Add(new FAndroid_DXTTargetPlatformControls(bIsClient, OutMap[EAndroidTextureFormatCategory::DXT]));
+				SinglePlatforms.Add(new FAndroid_ETC2TargetPlatformControls(bIsClient, OutMap[EAndroidTextureFormatCategory::ETC2]));
+
+				// these are used in NotifyMultiSelectedFormatsChanged, so track in another array
+				MultiPlatforms.Add(new FAndroid_MultiTargetPlatformControls(bIsClient, OutMulti));
+			}
+
+			// join the single and the multi into one
+			TargetPlatforms.Append(SinglePlatforms);
+			TargetPlatforms.Append(MultiPlatforms);
+
+			// set up the multi platforms now that we have all the other platforms ready to go
+			NotifyMultiSelectedFormatsChanged();
+		}
+	}
+
+	virtual void NotifyMultiSelectedFormatsChanged() override
+	{
+		for (FAndroid_MultiTargetPlatformControls* TP : MultiPlatforms)
+		{
+			TP->LoadFormats(SinglePlatforms);
+		}
+		// @todo multi needs to be passed this event!
+	}
+
+private:
+	/** Holds the specific types of target platforms for NotifyMultiSelectedFormatsChanged */
+	TArray<FAndroidTargetPlatformControls*> SinglePlatforms;
+	TArray<FAndroid_MultiTargetPlatformControls*> MultiPlatforms;
+};
+
+
+#undef LOCTEXT_NAMESPACE
+
+
+IMPLEMENT_MODULE(FAndroidTargetPlatformControlsModule, AndroidTargetPlatformControls);
\ No newline at end of file
diff --git a/Engine/Source/Developer/Android/AndroidTargetPlatformControls/Public/IAndroidTargetPlatformControlsModule.h b/Engine/Source/Developer/Android/AndroidTargetPlatformControls/Public/IAndroidTargetPlatformControlsModule.h
new file mode 100644
--- /dev/null
+++ b/Engine/Source/Developer/Android/AndroidTargetPlatformControls/Public/IAndroidTargetPlatformControlsModule.h
@@ -0,0 +1,19 @@
+// Copyright Epic Games, Inc. All Rights Reserved.
+
+#pragma  once 
+
+/*------------------------------------------------------------------------------------
+IAndroid_MultiTargetPlatformModule interface 
+------------------------------------------------------------------------------------*/
+
+#include "CoreMinimal.h"
+#include "Interfaces/ITargetPlatformControlsModule.h"
+
+class IAndroidTargetPlatformControlsModule : public ITargetPlatformControlsModule
+{
+public:
+	//
+	// Called by AndroidRuntimeSettings to notify us when the user changes the selected texture formats for the Multi format
+	//
+	virtual void NotifyMultiSelectedFormatsChanged() = 0;
+};
diff --git a/Engine/Source/Developer/Android/AndroidTargetPlatformSettings/AndroidTargetPlatformSettings.Build.cs b/Engine/Source/Developer/Android/AndroidTargetPlatformSettings/AndroidTargetPlatformSettings.Build.cs
new file mode 100644
--- /dev/null
+++ b/Engine/Source/Developer/Android/AndroidTargetPlatformSettings/AndroidTargetPlatformSettings.Build.cs
@@ -0,0 +1,27 @@
+// Copyright Epic Games, Inc. All Rights Reserved.
+
+using UnrealBuildTool;
+
+public class AndroidTargetPlatformSettings : ModuleRules
+{
+	public AndroidTargetPlatformSettings(ReadOnlyTargetRules Target) : base(Target)
+	{
+		UnsafeTypeCastWarningLevel = WarningLevel.Error;
+		BinariesSubFolder = "Android";
+
+		PrivateDependencyModuleNames.AddRange(
+			new string[]
+			{
+				"Core",
+				"CoreUObject",
+				"TargetPlatform",
+				"DesktopPlatform",
+			}
+		);
+
+		if (Target.bCompileAgainstEngine)
+		{
+			PrivateDependencyModuleNames.Add("Engine");
+		}
+	}
+}
diff --git a/Engine/Source/Developer/Android/AndroidTargetPlatformSettings/Private/AndroidTargetPlatformSettings.cpp b/Engine/Source/Developer/Android/AndroidTargetPlatformSettings/Private/AndroidTargetPlatformSettings.cpp
new file mode 100644
--- /dev/null
+++ b/Engine/Source/Developer/Android/AndroidTargetPlatformSettings/Private/AndroidTargetPlatformSettings.cpp
@@ -0,0 +1,148 @@
+// Copyright Epic Games, Inc. All Rights Reserved.
+
+#include "AndroidTargetPlatformSettings.h"
+#include "CoreMinimal.h"
+#include "UObject/NameTypes.h"
+#include "CoreTypes.h"
+#include "AnalyticsEventAttribute.h"
+
+FAndroidTargetPlatformSettings::FAndroidTargetPlatformSettings(const TCHAR* CookFlavor, const TCHAR* OverrideIniPlatformName)
+	: TTargetPlatformSettingsBase(CookFlavor, OverrideIniPlatformName)
+	, MobileShadingPath(0)
+	, bDistanceField(false)
+	, bMobileForwardEnableClusteredReflections(false)
+	, bMobileVirtualTextures(false)
+{
+#if WITH_ENGINE
+	TextureLODSettings = nullptr; // These are registered by the device profile system.
+	StaticMeshLODSettings.Initialize(this);
+	GetConfigSystem()->GetBool(TEXT("/Script/Engine.RendererSettings"), TEXT("r.DistanceFields"), bDistanceField, GEngineIni);
+	GetConfigSystem()->GetInt(TEXT("/Script/Engine.RendererSettings"), TEXT("r.Mobile.ShadingPath"), MobileShadingPath, GEngineIni);
+	GetConfigSystem()->GetBool(TEXT("/Script/Engine.RendererSettings"), TEXT("r.Mobile.Forward.EnableClusteredReflections"), bMobileForwardEnableClusteredReflections, GEngineIni);
+	GetConfigSystem()->GetBool(TEXT("/Script/Engine.RendererSettings"), TEXT("r.Mobile.VirtualTextures"), bMobileVirtualTextures, GEngineIni);
+#endif
+}
+
+bool FAndroidTargetPlatformSettings::SupportsES31() const
+{
+	// default no support for ES31
+	bool bBuildForES31 = false;
+#if WITH_ENGINE
+	GConfig->GetBool(TEXT("/Script/AndroidRuntimeSettings.AndroidRuntimeSettings"), TEXT("bBuildForES31"), bBuildForES31, GEngineIni);
+#endif
+	return bBuildForES31;
+}
+
+bool FAndroidTargetPlatformSettings::SupportsVulkan() const
+{
+	// default to not supporting Vulkan
+	bool bSupportsVulkan = false;
+#if WITH_ENGINE
+	GConfig->GetBool(TEXT("/Script/AndroidRuntimeSettings.AndroidRuntimeSettings"), TEXT("bSupportsVulkan"), bSupportsVulkan, GEngineIni);
+#endif
+	return bSupportsVulkan;
+}
+
+bool FAndroidTargetPlatformSettings::SupportsVulkanSM5() const
+{
+	// default to no support for VulkanSM5
+	bool bSupportsMobileVulkanSM5 = false;
+#if WITH_ENGINE
+	GConfig->GetBool(TEXT("/Script/AndroidRuntimeSettings.AndroidRuntimeSettings"), TEXT("bSupportsVulkanSM5"), bSupportsMobileVulkanSM5, GEngineIni);
+#endif
+	return bSupportsMobileVulkanSM5;
+}
+
+bool FAndroidTargetPlatformSettings::SupportsFeature(ETargetPlatformFeatures Feature) const
+{
+	switch (Feature)
+	{
+	case ETargetPlatformFeatures::Packaging:
+	case ETargetPlatformFeatures::DeviceOutputLog:
+		return true;
+
+	case ETargetPlatformFeatures::LowQualityLightmaps:
+	case ETargetPlatformFeatures::MobileRendering:
+		return SupportsES31() || SupportsVulkan();
+
+	case ETargetPlatformFeatures::HighQualityLightmaps:
+	case ETargetPlatformFeatures::DeferredRendering:
+		return SupportsVulkanSM5();
+
+	case ETargetPlatformFeatures::VirtualTextureStreaming:
+		// TODO: should it check r.VirtualTextures for SM5 renderer ?
+		return bMobileVirtualTextures;
+
+	case ETargetPlatformFeatures::DistanceFieldAO:
+		return UsesDistanceFields();
+
+
+	case ETargetPlatformFeatures::NormalmapLAEncodingMode:
+	{
+		static IConsoleVariable* CompressorCVar = IConsoleManager::Get().FindConsoleVariable(TEXT("cook.ASTCTextureCompressor"));
+		const bool bUsesARMCompressor = (CompressorCVar ? (CompressorCVar->GetInt() != 0) : false);
+		return SupportsTextureFormatCategory(EAndroidTextureFormatCategory::ASTC) && bUsesARMCompressor;
+	}
+
+	default:
+		break;
+	}
+
+	return TTargetPlatformSettingsBase<FAndroidPlatformProperties>::SupportsFeature(Feature);
+}
+
+#if WITH_ENGINE
+void FAndroidTargetPlatformSettings::GetReflectionCaptureFormats(TArray<FName>& OutFormats) const
+{
+	const bool bMobileDeferredShading = (MobileShadingPath == 1);
+
+	if (SupportsVulkanSM5() || bMobileDeferredShading || bMobileForwardEnableClusteredReflections)
+	{
+		// use Full HDR with SM5 and Mobile Deferred
+		OutFormats.Add(FName(TEXT("FullHDR")));
+	}
+
+	// always emit encoded
+	OutFormats.Add(FName(TEXT("EncodedHDR")));
+}
+const UTextureLODSettings& FAndroidTargetPlatformSettings::GetTextureLODSettings() const
+{
+	return *TextureLODSettings;
+}
+#endif
+
+void FAndroidTargetPlatformSettings::GetAllPossibleShaderFormats(TArray<FName>& OutFormats) const
+{
+	static FName NAME_SF_VULKAN_ES31_ANDROID(TEXT("SF_VULKAN_ES31_ANDROID"));
+	static FName NAME_GLSL_ES3_1_ANDROID(TEXT("GLSL_ES3_1_ANDROID"));
+	static FName NAME_SF_VULKAN_SM5_ANDROID(TEXT("SF_VULKAN_SM5_ANDROID"));
+
+	if (SupportsVulkan())
+	{
+		OutFormats.AddUnique(NAME_SF_VULKAN_ES31_ANDROID);
+	}
+
+	if (SupportsVulkanSM5())
+	{
+		OutFormats.AddUnique(NAME_SF_VULKAN_SM5_ANDROID);
+	}
+
+	if (SupportsES31())
+	{
+		OutFormats.AddUnique(NAME_GLSL_ES3_1_ANDROID);
+	}
+}
+
+void FAndroidTargetPlatformSettings::GetAllTargetedShaderFormats(TArray<FName>& OutFormats) const
+{
+	GetAllPossibleShaderFormats(OutFormats);
+}
+
+
+#if WITH_ENGINE
+const FStaticMeshLODSettings& FAndroidTargetPlatformSettings::GetStaticMeshLODSettings() const
+{
+	return StaticMeshLODSettings;
+}
+
+#endif
\ No newline at end of file
diff --git a/Engine/Source/Developer/Android/AndroidTargetPlatformSettings/Private/AndroidTargetPlatformSettingsModule.cpp b/Engine/Source/Developer/Android/AndroidTargetPlatformSettings/Private/AndroidTargetPlatformSettingsModule.cpp
new file mode 100644
--- /dev/null
+++ b/Engine/Source/Developer/Android/AndroidTargetPlatformSettings/Private/AndroidTargetPlatformSettingsModule.cpp
@@ -0,0 +1,69 @@
+// Copyright Epic Games, Inc. All Rights Reserved.
+
+#include "CoreMinimal.h"
+#include "AndroidTargetPlatformSettings.h"
+#include "Modules/ModuleManager.h"
+#include "Interfaces/ITargetPlatformModule.h"
+#include "Interfaces/ITargetPlatformSettingsModule.h"
+#include "Interfaces/ITargetPlatformControlsModule.h"
+#include "Android/AndroidPlatformProperties.h"
+#include "MultiPlatformTargetReceiptBuildWorkers.h"
+#include "IAndroidTargetPlatformSettingsModule.h"
+
+/**
+ * Module for the Androids target platform settings.
+ */
+class FAndroidTargetPlatformSettingsModule
+	: public IAndroidTargetPlatformSettingsModule
+{
+public:
+
+	/** Destructor. */
+	~FAndroidTargetPlatformSettingsModule() 
+	{
+	}
+
+public:
+	virtual void GetTargetPlatforms(TArray<ITargetPlatform*>& TargetPlatforms)
+	{
+
+	}
+	void GetSinglePlatformSettings(TMap<EAndroidTextureFormatCategory, ITargetPlatformSettings*>& OutMap)
+	{
+		OutMap = SinglePlatformSettings;
+	}
+
+	ITargetPlatformSettings* GetMultiPlatformSettings()
+	{
+		return MultiTargetPlatformSettings;
+	}
+
+protected:
+	virtual void GetTargetPlatformSettings(TArray<ITargetPlatformSettings*>& TargetPlatforms)
+	{
+		ITargetPlatformSettings* VanillaTP = new FAndroidTargetPlatformSettings();
+		TargetPlatforms.Add(VanillaTP);
+		SinglePlatformSettings.Add(EAndroidTextureFormatCategory::Count, VanillaTP);
+
+		ITargetPlatformSettings* DXTTP = new FAndroid_DXTTargetPlatformSettings();
+		TargetPlatforms.Add(DXTTP);
+		SinglePlatformSettings.Add(EAndroidTextureFormatCategory::DXT, DXTTP);
+
+		ITargetPlatformSettings* ETC2TP = new FAndroid_ETC2TargetPlatformSettings();
+		TargetPlatforms.Add(ETC2TP);
+		SinglePlatformSettings.Add(EAndroidTextureFormatCategory::ETC2, ETC2TP);
+
+		ITargetPlatformSettings*ASTCTP = new FAndroid_ASTCTargetPlatformSettings();
+		TargetPlatforms.Add(ASTCTP);
+		SinglePlatformSettings.Add(EAndroidTextureFormatCategory::ASTC, ASTCTP);
+
+		ITargetPlatformSettings* MultiTP = new FAndroid_MultiTargetPlatformSettings();
+		TargetPlatforms.Add(MultiTP);
+		MultiTargetPlatformSettings = MultiTP;
+	}
+private:
+	ITargetPlatformSettings* MultiTargetPlatformSettings;
+	TMap<EAndroidTextureFormatCategory, ITargetPlatformSettings*> SinglePlatformSettings;
+};
+
+IMPLEMENT_MODULE(FAndroidTargetPlatformSettingsModule, AndroidTargetPlatformSettings);
\ No newline at end of file
diff --git a/Engine/Source/Developer/Android/AndroidTargetPlatformSettings/Public/AndroidTargetPlatformSettings.h b/Engine/Source/Developer/Android/AndroidTargetPlatformSettings/Public/AndroidTargetPlatformSettings.h
new file mode 100644
--- /dev/null
+++ b/Engine/Source/Developer/Android/AndroidTargetPlatformSettings/Public/AndroidTargetPlatformSettings.h
@@ -0,0 +1,207 @@
+// Copyright Epic Games, Inc. All Rights Reserved.
+
+#pragma once
+#include "CoreMinimal.h"
+#include "Common/TargetPlatformSettingsBase.h"
+#include "Android/AndroidPlatformProperties.h"
+#include "Misc/ConfigCacheIni.h"
+#include "Misc/MonitoredProcess.h"
+#include "Containers/Ticker.h"
+
+#if WITH_ENGINE
+#include "StaticMeshResources.h"
+#endif // WITH_ENGINE
+
+
+/** Listed in order of priority...if device supports multiple formats, first format in list will be chosen */
+enum class EAndroidTextureFormatCategory
+{
+	DXT,
+	ETC2,
+	ASTC,
+
+	Count,
+};
+
+namespace AndroidTexFormat
+{
+	// Compressed Texture Formats
+	const static FName NameDXT1(TEXT("DXT1"));
+	const static FName NameDXT5(TEXT("DXT5"));
+	const static FName NameDXT5n(TEXT("DXT5n"));
+	const static FName NameAutoDXT(TEXT("AutoDXT"));
+	const static FName NameBC4(TEXT("BC4"));
+	const static FName NameBC5(TEXT("BC5"));
+	const static FName NameBC6H(TEXT("BC6H"));
+	const static FName NameBC7(TEXT("BC7"));
+
+	const static FName NameETC2_RGB(TEXT("ETC2_RGB"));
+	const static FName NameETC2_RGBA(TEXT("ETC2_RGBA"));
+	const static FName NameETC2_R11(TEXT("ETC2_R11"));
+	const static FName NameETC2_RG11(TEXT("ETC2_RG11"));
+	const static FName NameAutoETC2(TEXT("AutoETC2"));
+
+	const static FName NameAutoASTC(TEXT("ASTC_RGBAuto"));
+	const static FName NameASTC_NormalRG(TEXT("ASTC_NormalRG"));
+	// L+A mode suppoprted by ARM ASTC encoder 
+	const static FName NameASTC_NormalLA(TEXT("ASTC_NormalLA"));
+
+	// Uncompressed Texture Formats
+	const static FName NameBGRA8(TEXT("BGRA8"));
+	const static FName NameG8(TEXT("G8"));
+	const static FName NameRGBA16F(TEXT("RGBA16F"));
+	const static FName NameR16F(TEXT("R16F"));
+	const static FName NameG16(TEXT("G16"));
+
+	//A1RGB555 is mapped to RGB555A1, because OpenGL GL_RGB5_A1 only supports alpha on the lowest bit
+	const static FName NameA1RGB555(TEXT("A1RGB555"));
+	const static FName NameRGB555A1(TEXT("RGB555A1"));
+
+	const static FName GenericRemap[][2] =
+	{
+		{ NameA1RGB555,		NameRGB555A1			},
+		{ NameG16,			NameR16F				}, // GLES does not support R16Unorm, fallback all Android to R16F
+	};
+
+	static const FName NameASTC_RGB_HDR(TEXT("ASTC_RGB_HDR"));
+
+	const static FName ASTCRemap[][2] =
+	{
+		// Default format:		ASTC format:
+		{ NameDXT1,			FName(TEXT("ASTC_RGB"))		},
+		{ NameDXT5,			FName(TEXT("ASTC_RGBA"))	},
+		{ NameDXT5n,		FName(TEXT("ASTC_NormalAG"))},
+		{ NameBC5,			NameASTC_NormalRG			},
+		{ NameBC4,			NameETC2_R11				},
+		{ NameBC6H,			NameASTC_RGB_HDR			},
+		{ NameBC7,			FName(TEXT("ASTC_RGBA_HQ"))	},
+		{ NameAutoDXT,		NameAutoASTC				},
+	};
+
+	const static FName ETCRemap[][2] =
+	{
+		// Default format:	ETC2 format:
+		{ NameDXT1,			NameETC2_RGB	},
+		{ NameDXT5,			NameETC2_RGBA	},
+		{ NameDXT5n,		NameETC2_RGB	},
+		{ NameBC5,			NameETC2_RG11	},
+		{ NameBC4,			NameETC2_R11	},
+		{ NameBC6H,			NameRGBA16F		},
+		{ NameBC7,			NameETC2_RGBA	},
+		{ NameAutoDXT,		NameAutoETC2	},
+	};
+}
+
+class ANDROIDTARGETPLATFORMSETTINGS_API FAndroidTargetPlatformSettings : public TTargetPlatformSettingsBase<FAndroidPlatformProperties>
+{
+public:
+	FAndroidTargetPlatformSettings(const TCHAR* CookFlavor = nullptr, const TCHAR* OverrideIniPlatformName = nullptr);
+
+	/**
+	 * Destructor
+	 */
+	virtual ~FAndroidTargetPlatformSettings(){}
+
+	virtual bool SupportsTextureFormatCategory(EAndroidTextureFormatCategory Category) const
+	{
+		return true;
+	}
+
+	virtual bool SupportsFeature(ETargetPlatformFeatures Feature) const override;
+	virtual void GetAllPossibleShaderFormats(TArray<FName>& OutFormats) const override;
+	virtual void GetAllTargetedShaderFormats(TArray<FName>& OutFormats) const override;
+
+	virtual bool UsesRayTracing() const override
+	{
+		bool bEnableRayTracing = false;
+		GConfig->GetBool(TEXT("/Script/AndroidRuntimeSettings.AndroidRuntimeSettings"), TEXT("bEnableRayTracing"), bEnableRayTracing, GEngineIni);
+
+		return bEnableRayTracing;
+	}
+
+#if WITH_ENGINE
+	virtual void GetReflectionCaptureFormats(TArray<FName>& OutFormats) const override;
+	virtual const UTextureLODSettings& GetTextureLODSettings() const override;
+	virtual void RegisterTextureLODSettings(const UTextureLODSettings* InTextureLODSettings) override
+	{
+		TextureLODSettings = InTextureLODSettings;
+	}
+	virtual const class FStaticMeshLODSettings& GetStaticMeshLODSettings() const override;
+#endif
+	// query for rendering mode support
+	bool SupportsES31() const;
+	bool SupportsVulkan() const;
+	bool SupportsVulkanSM5() const;
+protected:
+
+	
+
+	// r.Mobile.ShadingPath value
+	int32 MobileShadingPath;
+
+	// true if DistanceField is enabled
+	bool bDistanceField;
+
+	// r.Mobile.Forward.EnableClusteredReflections value
+	bool bMobileForwardEnableClusteredReflections;
+
+	// r.Mobile.VirtualTextures value
+	bool bMobileVirtualTextures;
+#if WITH_ENGINE
+	// Holds a cache of the target LOD settings.
+	const UTextureLODSettings* TextureLODSettings;
+
+	// Holds the static mesh LOD settings.
+	FStaticMeshLODSettings StaticMeshLODSettings;
+#endif
+};
+
+class FAndroid_DXTTargetPlatformSettings : public FAndroidTargetPlatformSettings
+{
+public:
+	FAndroid_DXTTargetPlatformSettings() : FAndroidTargetPlatformSettings(TEXT("DXT"))
+	{
+	}
+
+	virtual bool SupportsTextureFormatCategory(EAndroidTextureFormatCategory Category) const override
+	{
+		return Category == EAndroidTextureFormatCategory::DXT;
+	}
+};
+
+class FAndroid_ASTCTargetPlatformSettings : public FAndroidTargetPlatformSettings
+{
+public:
+	FAndroid_ASTCTargetPlatformSettings() : FAndroidTargetPlatformSettings(TEXT("ASTC"))
+	{
+	}
+
+	virtual bool SupportsTextureFormatCategory(EAndroidTextureFormatCategory Category) const override
+	{
+		return Category == EAndroidTextureFormatCategory::ASTC;
+	}
+};
+
+
+class FAndroid_ETC2TargetPlatformSettings : public FAndroidTargetPlatformSettings
+{
+public:
+
+	FAndroid_ETC2TargetPlatformSettings() : FAndroidTargetPlatformSettings(TEXT("ETC2"))
+	{
+	}
+
+	virtual bool SupportsTextureFormatCategory(EAndroidTextureFormatCategory Category) const override
+	{
+		return Category == EAndroidTextureFormatCategory::ETC2;
+	}
+};
+
+class FAndroid_MultiTargetPlatformSettings : public FAndroidTargetPlatformSettings
+{
+
+public:
+	FAndroid_MultiTargetPlatformSettings() : FAndroidTargetPlatformSettings(TEXT("Multi"))
+	{
+	}
+};
diff --git a/Engine/Source/Developer/Android/AndroidTargetPlatformSettings/Public/IAndroidTargetPlatformSettingsModule.h b/Engine/Source/Developer/Android/AndroidTargetPlatformSettings/Public/IAndroidTargetPlatformSettingsModule.h
new file mode 100644
--- /dev/null
+++ b/Engine/Source/Developer/Android/AndroidTargetPlatformSettings/Public/IAndroidTargetPlatformSettingsModule.h
@@ -0,0 +1,18 @@
+// Copyright Epic Games, Inc. All Rights Reserved.
+
+#pragma  once 
+
+/*------------------------------------------------------------------------------------
+IAndroidTargetPlatformSettingsModule interface
+------------------------------------------------------------------------------------*/
+
+#include "CoreMinimal.h"
+#include "Interfaces/ITargetPlatformSettingsModule.h"
+#include "AndroidTargetPlatformSettings.h"
+
+class IAndroidTargetPlatformSettingsModule : public ITargetPlatformSettingsModule
+{
+public:
+	virtual void GetSinglePlatformSettings(TMap<EAndroidTextureFormatCategory, ITargetPlatformSettings*>& OutMap) = 0;
+	virtual ITargetPlatformSettings* GetMultiPlatformSettings() = 0;
+};
diff --git a/Engine/Source/Developer/CookOnTheFlyNetServer/Private/PlatformProtocolServer.cpp b/Engine/Source/Developer/CookOnTheFlyNetServer/Private/PlatformProtocolServer.cpp
--- a/Engine/Source/Developer/CookOnTheFlyNetServer/Private/PlatformProtocolServer.cpp
+++ b/Engine/Source/Developer/CookOnTheFlyNetServer/Private/PlatformProtocolServer.cpp
@@ -90,7 +90,7 @@ public:
 
 	FString GetName() const
 	{
-		return FString::Printf(TEXT("%s (%s)"), *Device->GetName(), *Device->GetTargetPlatform().PlatformName());
+		return FString::Printf(TEXT("%s (%s)"), *Device->GetName(), *Device->GetPlatformControls().PlatformName());
 	}
 
 	ITargetDevicePtr GetDevice() const
diff --git a/Engine/Source/Developer/DeviceManager/Private/Widgets/Details/SDeviceDetails.cpp b/Engine/Source/Developer/DeviceManager/Private/Widgets/Details/SDeviceDetails.cpp
--- a/Engine/Source/Developer/DeviceManager/Private/Widgets/Details/SDeviceDetails.cpp
+++ b/Engine/Source/Developer/DeviceManager/Private/Widgets/Details/SDeviceDetails.cpp
@@ -142,23 +142,24 @@ void SDeviceDetails::Construct(const FArguments& InArgs, const TSharedRef<FDevic
 
 			if (TargetDevice.IsValid())
 			{
-				const ITargetPlatform& TargetPlatform = TargetDevice->GetTargetPlatform();
+				const ITargetPlatformSettings& TargetPlatformSettings = TargetDevice->GetPlatformSettings();
+				const ITargetPlatformControls& TargetPlatformControls = TargetDevice->GetPlatformControls();
 
 				// platform features
-				FeatureList.Add(MakeShareable(new FDeviceDetailsFeature(TEXT("AudioStreaming"), TargetPlatform.SupportsFeature(ETargetPlatformFeatures::AudioStreaming))));
-				FeatureList.Add(MakeShareable(new FDeviceDetailsFeature(TEXT("DistanceFieldShadows"), TargetPlatform.SupportsFeature(ETargetPlatformFeatures::DistanceFieldShadows))));
-				FeatureList.Add(MakeShareable(new FDeviceDetailsFeature(TEXT("DistanceFieldAO"), TargetPlatform.SupportsFeature(ETargetPlatformFeatures::DistanceFieldAO))));
-				FeatureList.Add(MakeShareable(new FDeviceDetailsFeature(TEXT("GrayscaleSRGB"), TargetPlatform.SupportsFeature(ETargetPlatformFeatures::GrayscaleSRGB))));
-				FeatureList.Add(MakeShareable(new FDeviceDetailsFeature(TEXT("HighQualityLightmaps"), TargetPlatform.SupportsFeature(ETargetPlatformFeatures::HighQualityLightmaps))));
-				FeatureList.Add(MakeShareable(new FDeviceDetailsFeature(TEXT("LowQualityLightmaps"), TargetPlatform.SupportsFeature(ETargetPlatformFeatures::LowQualityLightmaps))));
-				FeatureList.Add(MakeShareable(new FDeviceDetailsFeature(TEXT("MultipleGameInstances"), TargetPlatform.SupportsFeature(ETargetPlatformFeatures::MultipleGameInstances))));
-				FeatureList.Add(MakeShareable(new FDeviceDetailsFeature(TEXT("Packaging"), TargetPlatform.SupportsFeature(ETargetPlatformFeatures::Packaging))));
-				FeatureList.Add(MakeShareable(new FDeviceDetailsFeature(TEXT("SdkConnectDisconnect"), TargetPlatform.SupportsFeature(ETargetPlatformFeatures::SdkConnectDisconnect))));
-				FeatureList.Add(MakeShareable(new FDeviceDetailsFeature(TEXT("TextureStreaming"), TargetPlatform.SupportsFeature(ETargetPlatformFeatures::TextureStreaming))));
-				FeatureList.Add(MakeShareable(new FDeviceDetailsFeature(TEXT("UserCredentials"), TargetPlatform.SupportsFeature(ETargetPlatformFeatures::UserCredentials))));
-				FeatureList.Add(MakeShareable(new FDeviceDetailsFeature(TEXT("MobileRendering"), TargetPlatform.SupportsFeature(ETargetPlatformFeatures::MobileRendering))));
-				FeatureList.Add(MakeShareable(new FDeviceDetailsFeature(TEXT("DeferredRendering"), TargetPlatform.SupportsFeature(ETargetPlatformFeatures::DeferredRendering))));
-				FeatureList.Add(MakeShareable(new FDeviceDetailsFeature(TEXT("LumenGI"), TargetPlatform.SupportsFeature(ETargetPlatformFeatures::LumenGI))));
+				FeatureList.Add(MakeShareable(new FDeviceDetailsFeature(TEXT("AudioStreaming"), TargetPlatformSettings.SupportsFeature(ETargetPlatformFeatures::AudioStreaming))));
+				FeatureList.Add(MakeShareable(new FDeviceDetailsFeature(TEXT("DistanceFieldShadows"), TargetPlatformSettings.SupportsFeature(ETargetPlatformFeatures::DistanceFieldShadows))));
+				FeatureList.Add(MakeShareable(new FDeviceDetailsFeature(TEXT("DistanceFieldAO"), TargetPlatformSettings.SupportsFeature(ETargetPlatformFeatures::DistanceFieldAO))));
+				FeatureList.Add(MakeShareable(new FDeviceDetailsFeature(TEXT("GrayscaleSRGB"), TargetPlatformSettings.SupportsFeature(ETargetPlatformFeatures::GrayscaleSRGB))));
+				FeatureList.Add(MakeShareable(new FDeviceDetailsFeature(TEXT("HighQualityLightmaps"), TargetPlatformSettings.SupportsFeature(ETargetPlatformFeatures::HighQualityLightmaps))));
+				FeatureList.Add(MakeShareable(new FDeviceDetailsFeature(TEXT("LowQualityLightmaps"), TargetPlatformSettings.SupportsFeature(ETargetPlatformFeatures::LowQualityLightmaps))));
+				FeatureList.Add(MakeShareable(new FDeviceDetailsFeature(TEXT("MultipleGameInstances"), TargetPlatformSettings.SupportsFeature(ETargetPlatformFeatures::MultipleGameInstances))));
+				FeatureList.Add(MakeShareable(new FDeviceDetailsFeature(TEXT("Packaging"), TargetPlatformSettings.SupportsFeature(ETargetPlatformFeatures::Packaging))));
+				FeatureList.Add(MakeShareable(new FDeviceDetailsFeature(TEXT("SdkConnectDisconnect"), TargetPlatformSettings.SupportsFeature(ETargetPlatformFeatures::SdkConnectDisconnect))));
+				FeatureList.Add(MakeShareable(new FDeviceDetailsFeature(TEXT("TextureStreaming"), TargetPlatformSettings.SupportsFeature(ETargetPlatformFeatures::TextureStreaming))));
+				FeatureList.Add(MakeShareable(new FDeviceDetailsFeature(TEXT("UserCredentials"), TargetPlatformSettings.SupportsFeature(ETargetPlatformFeatures::UserCredentials))));
+				FeatureList.Add(MakeShareable(new FDeviceDetailsFeature(TEXT("MobileRendering"), TargetPlatformSettings.SupportsFeature(ETargetPlatformFeatures::MobileRendering))));
+				FeatureList.Add(MakeShareable(new FDeviceDetailsFeature(TEXT("DeferredRendering"), TargetPlatformSettings.SupportsFeature(ETargetPlatformFeatures::DeferredRendering))));
+				FeatureList.Add(MakeShareable(new FDeviceDetailsFeature(TEXT("LumenGI"), TargetPlatformSettings.SupportsFeature(ETargetPlatformFeatures::LumenGI))));
 
 				// device features
 				FeatureList.Add(MakeShareable(new FDeviceDetailsFeature(TEXT("MultiLaunch"), TargetDevice->SupportsFeature(ETargetDeviceFeatures::MultiLaunch))));
@@ -168,7 +169,7 @@ void SDeviceDetails::Construct(const FArguments& InArgs, const TSharedRef<FDevic
 				FeatureList.Add(MakeShareable(new FDeviceDetailsFeature(TEXT("Reboot"), TargetDevice->SupportsFeature(ETargetDeviceFeatures::Reboot))));
 
 				// create custom widget for the platform and place it in the panel
-				TSharedPtr<SWidget> CustomWidget = TargetPlatform.GetCustomWidgetCreator()->CreateDeviceInfoWidget(TargetPlatform.PlatformName(), TargetDevice);
+				TSharedPtr<SWidget> CustomWidget = TargetPlatformControls.GetCustomWidgetCreator()->CreateDeviceInfoWidget(TargetPlatformControls.PlatformName(), TargetDevice);
 				if (CustomWidget)
 				{
 					CustomPlatformWidgetPanel->SetContent(CustomWidget.ToSharedRef());
diff --git a/Engine/Source/Developer/DeviceManager/Private/Widgets/SDeviceManager.cpp b/Engine/Source/Developer/DeviceManager/Private/Widgets/SDeviceManager.cpp
--- a/Engine/Source/Developer/DeviceManager/Private/Widgets/SDeviceManager.cpp
+++ b/Engine/Source/Developer/DeviceManager/Private/Widgets/SDeviceManager.cpp
@@ -270,7 +270,7 @@ bool SDeviceManager::HandleConnectActionCanExecute()
 	{
 		ITargetDevicePtr TargetDevice = DeviceService->GetDevice();
 
-		if (TargetDevice.IsValid() && TargetDevice->GetTargetPlatform().SupportsFeature(ETargetPlatformFeatures::SdkConnectDisconnect))
+		if (TargetDevice.IsValid() && TargetDevice->GetPlatformSettings().SupportsFeature(ETargetPlatformFeatures::SdkConnectDisconnect))
 		{
 			return !TargetDevice->IsConnected();
 		}
@@ -307,7 +307,7 @@ bool SDeviceManager::HandleDisconnectActionCanExecute()
 	{
 		ITargetDevicePtr TargetDevice = DeviceService->GetDevice();
 
-		if (TargetDevice.IsValid() && TargetDevice->GetTargetPlatform().SupportsFeature(ETargetPlatformFeatures::SdkConnectDisconnect))
+		if (TargetDevice.IsValid() && TargetDevice->GetPlatformSettings().SupportsFeature(ETargetPlatformFeatures::SdkConnectDisconnect))
 		{
 			return TargetDevice->IsConnected();
 		}
diff --git a/Engine/Source/Developer/OutputLog/Private/SDeviceOutputLog.cpp b/Engine/Source/Developer/OutputLog/Private/SDeviceOutputLog.cpp
--- a/Engine/Source/Developer/OutputLog/Private/SDeviceOutputLog.cpp
+++ b/Engine/Source/Developer/OutputLog/Private/SDeviceOutputLog.cpp
@@ -277,7 +277,7 @@ void SDeviceOutputLog::AddDeviceEntry(ITargetDeviceRef TargetDevice)
 		return;
 	}
 	using namespace PlatformInfo;
-	FName DeviceIconStyleName = TargetDevice->GetTargetPlatform().GetPlatformInfo().GetIconStyleName(EPlatformIconSize::Normal);
+	FName DeviceIconStyleName = TargetDevice->GetPlatformControls().GetPlatformInfo().GetIconStyleName(EPlatformIconSize::Normal);
 	
 	TSharedPtr<FTargetDeviceEntry> DeviceEntry = MakeShareable(new FTargetDeviceEntry());
 	
diff --git a/Engine/Source/Developer/TargetDeviceServices/Private/Services/TargetDeviceService.cpp b/Engine/Source/Developer/TargetDeviceServices/Private/Services/TargetDeviceService.cpp
--- a/Engine/Source/Developer/TargetDeviceServices/Private/Services/TargetDeviceService.cpp
+++ b/Engine/Source/Developer/TargetDeviceServices/Private/Services/TargetDeviceService.cpp
@@ -14,7 +14,8 @@
 #include "Misc/ConfigCacheIni.h"
 #include "PlatformInfo.h"
 #include "Serialization/Archive.h"
-
+#include "Interfaces/ITargetPlatformSettings.h"
+#include "Interfaces/ITargetPlatformControls.h"
 #include "TargetDeviceServiceMessages.h"
 
 LLM_DECLARE_TAG(TargetDeviceProxyManager);
@@ -30,7 +31,7 @@ struct FVariantSortCallback
 		ITargetDevicePtr APtr = A.Pin();
 		ITargetDevicePtr BPtr = B.Pin();
 
-		return APtr->GetTargetPlatform().GetVariantPriority() > BPtr->GetTargetPlatform().GetVariantPriority();
+		return APtr->GetPlatformControls().GetVariantPriority() > BPtr->GetPlatformControls().GetVariantPriority();
 	}
 };
 
@@ -80,13 +81,13 @@ void FTargetDeviceService::AddTargetDevice(TSharedPtr<ITargetDevice, ESPMode::Th
 		return;
 	}
 
-	FName Variant = FName(InDevice->GetTargetPlatform().PlatformName().GetCharArray().GetData());
+	FName Variant = FName(InDevice->GetPlatformControls().PlatformName().GetCharArray().GetData());
 
 	if (DevicePlatformName == NAME_None)
 	{
 		// If this seems nasty your right!
 		// This is just one more nastiness in this class due to the fact that we intend to refactor the target platform stuff as a separate task.
-		const PlatformInfo::FTargetPlatformInfo& Info = InDevice->GetTargetPlatform().GetTargetPlatformInfo();
+		const PlatformInfo::FTargetPlatformInfo& Info = InDevice->GetPlatformControls().GetTargetPlatformInfo();
 		const PlatformInfo::FTargetPlatformInfo* VanillaInfo = Info.VanillaInfo;
 
 		DevicePlatformName = Info.Name;
@@ -212,7 +213,7 @@ void FTargetDeviceService::RemoveTargetDevice(TSharedPtr<ITargetDevice, ESPMode:
 		return;
 	}
 
-	FName Variant = FName(InDevice->GetTargetPlatform().PlatformName().GetCharArray().GetData());
+	FName Variant = FName(InDevice->GetPlatformControls().PlatformName().GetCharArray().GetData());
 
 	TargetDevicePtrs.Remove(Variant);
 
@@ -415,8 +416,8 @@ void FTargetDeviceService::HandlePingMessage(const FTargetDeviceServicePing& InM
 
 	if (DefaultDevice.IsValid())
 	{
-		const FString& PlatformName = DefaultDevice->GetTargetPlatform().PlatformName();
-		const PlatformInfo::FTargetPlatformInfo* VanillaInfo = DefaultDevice->GetTargetPlatform().GetTargetPlatformInfo().VanillaInfo;
+		const FString& PlatformName = DefaultDevice->GetPlatformControls().PlatformName();
+		const PlatformInfo::FTargetPlatformInfo* VanillaInfo = DefaultDevice->GetPlatformControls().GetTargetPlatformInfo().VanillaInfo;
 
 		FTargetDeviceServicePong* Message = FMessageEndpoint::MakeMessage<FTargetDeviceServicePong>();
 
@@ -436,8 +437,8 @@ void FTargetDeviceService::HandlePingMessage(const FTargetDeviceServicePing& InM
 		Message->SupportsPowerOff = DefaultDevice->SupportsFeature(ETargetDeviceFeatures::PowerOff);
 		Message->SupportsPowerOn = DefaultDevice->SupportsFeature(ETargetDeviceFeatures::PowerOn);
 		Message->SupportsReboot = DefaultDevice->SupportsFeature(ETargetDeviceFeatures::Reboot);
-		Message->SupportsVariants = DefaultDevice->GetTargetPlatform().SupportsVariants();
-		Message->DefaultVariant = FName(DefaultDevice->GetTargetPlatform().PlatformName().GetCharArray().GetData());
+		Message->SupportsVariants = DefaultDevice->GetPlatformControls().SupportsVariants();
+		Message->DefaultVariant = FName(DefaultDevice->GetPlatformControls().PlatformName().GetCharArray().GetData());
 
 		// Check if we should also create an aggregate (All_<platform>_devices_on_<host>) proxy
 		Message->Aggregated = DefaultDevice->IsPlatformAggregated();
@@ -451,13 +452,13 @@ void FTargetDeviceService::HandlePingMessage(const FTargetDeviceServicePing& InM
 		for (auto TargetDeviceIt = TargetDevicePtrs.CreateIterator(); TargetDeviceIt; ++TargetDeviceIt, ++Index)
 		{
 			const ITargetDevicePtr& TargetDevice = TargetDeviceIt.Value().Pin();
-			const PlatformInfo::FTargetPlatformInfo& Info = TargetDevice->GetTargetPlatform().GetTargetPlatformInfo();
+			const PlatformInfo::FTargetPlatformInfo& Info = TargetDevice->GetPlatformControls().GetTargetPlatformInfo();
 
 			FTargetDeviceVariant& Variant = Message->Variants[Index];
 
 			Variant.DeviceID = TargetDevice->GetId().ToString();
 			Variant.VariantName = TargetDeviceIt.Key();
-			Variant.TargetPlatformName = TargetDevice->GetTargetPlatform().PlatformName();
+			Variant.TargetPlatformName = TargetDevice->GetPlatformControls().PlatformName();
 			Variant.TargetPlatformId = Info.Name;
 			Variant.VanillaPlatformId = Info.VanillaInfo->Name;
 			Variant.PlatformDisplayName = Info.DisplayName.ToString();
diff --git a/Engine/Source/Developer/TargetPlatform/Private/TargetPlatformControlsBase.cpp b/Engine/Source/Developer/TargetPlatform/Private/TargetPlatformControlsBase.cpp
new file mode 100644
--- /dev/null
+++ b/Engine/Source/Developer/TargetPlatform/Private/TargetPlatformControlsBase.cpp
@@ -0,0 +1,297 @@
+// Copyright Epic Games, Inc. All Rights Reserved.
+
+#include "Common/TargetPlatformControlsBase.h"
+#include "Interfaces/IProjectBuildMutatorFeature.h"
+#include "ProjectDescriptor.h"
+#include "Interfaces/IProjectManager.h"
+#include "Interfaces/IPluginManager.h"
+#include "AnalyticsEventAttribute.h"
+#include "Sound/AudioFormatSettings.h"
+#include "DeviceBrowserDefaultPlatformWidgetCreator.h"
+#include "Features/IModularFeatures.h"
+#include "Misc/App.h"
+
+#define LOCTEXT_NAMESPACE "TargetPlatform"
+
+void FTargetPlatformControlsBase::GetPlatformSpecificProjectAnalytics(TArray<FAnalyticsEventAttribute>& AnalyticsParamArray) const
+{
+	AppendAnalyticsEventAttributeArray(AnalyticsParamArray,
+		TEXT("UsesDistanceFields"), TargetPlatformSettings->UsesDistanceFields(),
+		TEXT("UsesForwardShading"), TargetPlatformSettings->UsesForwardShading()
+	);
+}
+
+void FTargetPlatformControlsBase::AppendAnalyticsEventConfigBool(TArray<FAnalyticsEventAttribute>& AnalyticsParamArray, const TCHAR* ConfigSection, const TCHAR* ConfigKey, const FString& IniFileName, const TCHAR* AnalyticsKeyNameOverride)
+{
+	bool ConfigValue;
+	if (GConfig->GetBool(ConfigSection, ConfigKey, ConfigValue, IniFileName))
+	{
+		AnalyticsParamArray.Add(FAnalyticsEventAttribute(AnalyticsKeyNameOverride ? AnalyticsKeyNameOverride : ConfigKey, ConfigValue));
+	}
+}
+
+void FTargetPlatformControlsBase::AppendAnalyticsEventConfigInt(TArray<FAnalyticsEventAttribute>& AnalyticsParamArray, const TCHAR* ConfigSection, const TCHAR* ConfigKey, const FString& IniFileName, const TCHAR* AnalyticsKeyNameOverride)
+{
+	int32 ConfigValue;
+	if (GConfig->GetInt(ConfigSection, ConfigKey, ConfigValue, IniFileName))
+	{
+		AnalyticsParamArray.Add(FAnalyticsEventAttribute(AnalyticsKeyNameOverride ? AnalyticsKeyNameOverride : ConfigKey, ConfigValue));
+	}
+}
+
+void FTargetPlatformControlsBase::AppendAnalyticsEventConfigFloat(TArray<FAnalyticsEventAttribute>& AnalyticsParamArray, const TCHAR* ConfigSection, const TCHAR* ConfigKey, const FString& IniFileName, const TCHAR* AnalyticsKeyNameOverride)
+{
+	float ConfigValue;
+	if (GConfig->GetFloat(ConfigSection, ConfigKey, ConfigValue, IniFileName))
+	{
+		AnalyticsParamArray.Add(FAnalyticsEventAttribute(AnalyticsKeyNameOverride ? AnalyticsKeyNameOverride : ConfigKey, ConfigValue));
+	}
+}
+
+void FTargetPlatformControlsBase::AppendAnalyticsEventConfigString(TArray<FAnalyticsEventAttribute>& AnalyticsParamArray, const TCHAR* ConfigSection, const TCHAR* ConfigKey, const FString& IniFileName, const TCHAR* AnalyticsKeyNameOverride)
+{
+	FString ConfigValue;
+	if (GConfig->GetString(ConfigSection, ConfigKey, ConfigValue, IniFileName))
+	{
+		AnalyticsParamArray.Add(FAnalyticsEventAttribute(AnalyticsKeyNameOverride ? AnalyticsKeyNameOverride : ConfigKey, ConfigValue));
+	}
+}
+
+void FTargetPlatformControlsBase::AppendAnalyticsEventConfigArray(TArray<FAnalyticsEventAttribute>& AnalyticsParamArray, const TCHAR* ConfigSection, const TCHAR* ConfigKey, const FString& IniFileName, const TCHAR* AnalyticsKeyNameOverride)
+{
+	TArray<FString> ConfigValue;
+	if (GConfig->GetArray(ConfigSection, ConfigKey, ConfigValue, IniFileName))
+	{
+		AnalyticsParamArray.Add(FAnalyticsEventAttribute(AnalyticsKeyNameOverride ? AnalyticsKeyNameOverride : ConfigKey, ConfigValue));
+	}
+}
+
+bool FTargetPlatformControlsBase::IsPluginEnabledForTarget(const IPlugin& Plugin, const FProjectDescriptor* Project, const FString& Platform, EBuildConfiguration Configuration, EBuildTargetType TargetType)
+{
+	if (!Plugin.GetDescriptor().SupportsTargetPlatform(Platform))
+	{
+		return false;
+	}
+
+	// TODO: Support transitive calculation of per-platform disabling for plugins.
+	// Plugins can reference other plugins, and it would be nice to be able to automatically disable for platform X
+	// plugins that are only referenced through another plugin that is disabled for platform X.
+	// For the time-being, to disable a transitively referenced plugin per-platform, the project has to 
+	// directly include the plugin.
+	IPluginManager& PluginManager = IPluginManager::Get();
+
+	if (Project != nullptr)
+	{
+		const FString& PluginName = Plugin.GetName();
+		const FPluginReferenceDescriptor* PluginReference = Project->Plugins.FindByPredicate(
+			[&PluginName](const FPluginReferenceDescriptor& ExistingReference)
+			{
+				return ExistingReference.Name == PluginName;
+			});
+		if (PluginReference)
+		{
+			// TODO: Remove this workaround for indirect plugin references. A project can mark a plugin as
+			// "Enabled": false, but that merely prevents a direct reference, and the plugin might be referenced and
+			// enabled by other plugins. PluginReference->IsEnabledForPlatform, IsEnabledForTargetConfiguration, and
+			// IsEnabledForTarget will all return false in that case, even though the plugin is actually enabled.
+			// Other systems using IPluginManager::Get().GetEnabledPlugins will disagree with the disabled result.
+			// To workaround it, when we detect the case of a disabled plugin reference for a plugin that is
+			// indirectly enabled, we treat it as having all platforms enabled.
+			// To fix it properly, we will need to have plugins track for which platforms they are enabled,
+			// and query the pluginmanager here instead of querying the PluginReference directly.
+			if (PluginReference->bEnabled || PluginReference->bEnabled == Plugin.IsEnabled())
+			{
+				bool bEnabledForProject = PluginReference->IsEnabledForPlatform(Platform) &&
+					(Configuration == EBuildConfiguration::Unknown || PluginReference->IsEnabledForTargetConfiguration(Configuration)) &&
+					PluginReference->IsEnabledForTarget(TargetType);
+				if (!bEnabledForProject)
+				{
+					return false;
+				}
+			}
+		}
+	}
+	return true;
+}
+
+bool FTargetPlatformControlsBase::RequiresTempTarget(bool bProjectHasCode, EBuildConfiguration Configuration, bool bRequiresAssetNativization, FText& OutReason) const
+{
+	// check to see if we already have a Target.cs file
+	if (bProjectHasCode)
+	{
+		return false;
+	}
+
+	// check if asset nativization is enabled
+	if (bRequiresAssetNativization)
+	{
+		OutReason = LOCTEXT("TempTarget_Nativization", "asset nativization is enabled");
+		return true;
+	}
+
+	// check to see if any projectmutator modular features are available
+	for (IProjectBuildMutatorFeature* Feature : IModularFeatures::Get().GetModularFeatureImplementations<IProjectBuildMutatorFeature>(PROJECT_BUILD_MUTATOR_FEATURE))
+	{
+		if (Feature->RequiresProjectBuild(PlatformInfo->Name, OutReason))
+		{
+			return true;
+		}
+	}
+
+	// check the target platforms for any differences in build settings or additional plugins
+	const FProjectDescriptor* Project = IProjectManager::Get().GetCurrentProject();
+	if (!FApp::IsEngineInstalled() && !HasDefaultBuildSettings())
+	{
+		OutReason = LOCTEXT("TempTarget_NonDefaultBuildConfig", "project has non-default build configuration");
+		return true;
+	}
+
+	// check if there's a non-default plugin change
+	FText Reason;
+	if (IPluginManager::Get().RequiresTempTargetForCodePlugin(Project, GetPlatformInfo().UBTPlatformString, Configuration, PlatformInfo->PlatformType, Reason))
+	{
+		OutReason = Reason;
+		return true;
+	}
+
+	return false;
+}
+
+bool FTargetPlatformControlsBase::IsEnabledForPlugin(const IPlugin& Plugin) const
+{
+	const FProjectDescriptor* Project = IProjectManager::Get().GetCurrentProject();
+	return IsPluginEnabledForTarget(Plugin, Project, GetPlatformInfo().UBTPlatformString, EBuildConfiguration::Unknown,
+		GetRuntimePlatformType());
+}
+
+TSharedPtr<IDeviceManagerCustomPlatformWidgetCreator> FTargetPlatformControlsBase::GetCustomWidgetCreator() const
+{
+	static TSharedPtr<FDeviceBrowserDefaultPlatformWidgetCreator> DefaultWidgetCreator = MakeShared<FDeviceBrowserDefaultPlatformWidgetCreator>();
+	return DefaultWidgetCreator;
+}
+
+bool FTargetPlatformControlsBase::HasDefaultBuildSettings() const
+{
+	// first check default build settings for all platforms
+	TArray<FString> BoolKeys, IntKeys, StringKeys, BuildKeys;
+	BuildKeys.Add(TEXT("bCompileApex"));
+	BuildKeys.Add(TEXT("bCompileICU"));
+	BuildKeys.Add(TEXT("bCompileSimplygon"));
+	BuildKeys.Add(TEXT("bCompileSimplygonSSF"));
+	BuildKeys.Add(TEXT("bCompileRecast"));
+	BuildKeys.Add(TEXT("bCompileSpeedTree"));
+	BuildKeys.Add(TEXT("bCompileWithPluginSupport"));
+	BuildKeys.Add(TEXT("bCompilePhysXVehicle"));
+	BuildKeys.Add(TEXT("bCompileFreeType"));
+	BuildKeys.Add(TEXT("bCompileForSize"));
+	BuildKeys.Add(TEXT("bCompileCEF3"));
+	BuildKeys.Add(TEXT("bCompileCustomSQLitePlatform"));
+
+	if (!DoProjectSettingsMatchDefault(TargetPlatformSettings->IniPlatformName(), TEXT("/Script/BuildSettings.BuildSettings"), &BuildKeys, nullptr, nullptr))
+	{
+		return false;
+	}
+
+	FString PlatformSection;
+	GetBuildProjectSettingKeys(PlatformSection, BoolKeys, IntKeys, StringKeys);
+
+	if (!DoProjectSettingsMatchDefault(TargetPlatformSettings->IniPlatformName(), PlatformSection, &BoolKeys, &IntKeys, &StringKeys))
+	{
+		return false;
+	}
+
+	return true;
+}
+
+bool FTargetPlatformControlsBase::DoProjectSettingsMatchDefault(const FString& InPlatformName, const FString& InSection, const TArray<FString>* InBoolKeys, const TArray<FString>* InIntKeys, const TArray<FString>* InStringKeys)
+{
+	FConfigFile ProjIni;
+	FConfigFile DefaultIni;
+	FConfigCacheIni::LoadLocalIniFile(ProjIni, TEXT("Engine"), true, *InPlatformName, true);
+	FConfigCacheIni::LoadExternalIniFile(DefaultIni, TEXT("Engine"), *FPaths::EngineConfigDir(), *FPaths::EngineConfigDir(), true, NULL, true);
+
+	if (InBoolKeys != NULL)
+	{
+		for (int Index = 0; Index < InBoolKeys->Num(); ++Index)
+		{
+			FString Default(TEXT("False")), Project(TEXT("False"));
+			DefaultIni.GetString(*InSection, *((*InBoolKeys)[Index]), Default);
+			ProjIni.GetString(*InSection, *((*InBoolKeys)[Index]), Project);
+			if (Default.Compare(Project, ESearchCase::IgnoreCase))
+			{
+				return false;
+			}
+		}
+	}
+
+	if (InIntKeys != NULL)
+	{
+		for (int Index = 0; Index < InIntKeys->Num(); ++Index)
+		{
+			int64 Default(0), Project(0);
+			DefaultIni.GetInt64(*InSection, *((*InIntKeys)[Index]), Default);
+			ProjIni.GetInt64(*InSection, *((*InIntKeys)[Index]), Project);
+			if (Default != Project)
+			{
+				return false;
+			}
+		}
+	}
+
+	if (InStringKeys != NULL)
+	{
+		for (int Index = 0; Index < InStringKeys->Num(); ++Index)
+		{
+			FString Default(TEXT("False")), Project(TEXT("False"));
+			DefaultIni.GetString(*InSection, *((*InStringKeys)[Index]), Default);
+			ProjIni.GetString(*InSection, *((*InStringKeys)[Index]), Project);
+			if (Default.Compare(Project, ESearchCase::IgnoreCase))
+			{
+				return false;
+			}
+		}
+	}
+
+	return true;
+}
+
+FTargetPlatformControlsBase::FTargetPlatformControlsBase(const PlatformInfo::FTargetPlatformInfo* const InPlatformInfo, ITargetPlatformSettings* TargetPlatformSettings)
+	: ITargetPlatformControls(TargetPlatformSettings)
+	, PlatformInfo(InPlatformInfo)
+{
+	checkf(PlatformInfo, TEXT("Null PlatformInfo was passed to FTargetPlatformBase. Check the static IsUsable function before creating this object. See FWindowsTargetPlatformModule::GetTargetPlatform()"));
+
+	PlatformOrdinal = AssignPlatformOrdinal(*this);
+
+#if WITH_ENGINE
+	// Build Audio Format Settings, Using long form equiv of GetConfigSysten to avoid calling a virtual
+	AudioFormatSettings = MakePimpl<Audio::FAudioFormatSettings>(
+		FConfigCacheIni::ForPlatform(InPlatformInfo->IniPlatformName), GEngineIni, InPlatformInfo->IniPlatformName.ToString());
+#endif //WITH_ENGINE
+
+}
+
+#if WITH_ENGINE
+
+const Audio::FAudioFormatSettings& FTargetPlatformControlsBase::GetAudioFormatSettings() const
+{
+	check(AudioFormatSettings.IsValid())
+		return *AudioFormatSettings;
+}
+
+FName FTargetPlatformControlsBase::GetWaveFormat(const class USoundWave* InWave) const
+{
+	return GetAudioFormatSettings().GetWaveFormat(InWave);
+}
+
+void FTargetPlatformControlsBase::GetAllWaveFormats(TArray<FName>& OutFormats) const
+{
+	GetAudioFormatSettings().GetAllWaveFormats(OutFormats);
+}
+
+void FTargetPlatformControlsBase::GetWaveFormatModuleHints(TArray<FName>& OutModuleNames) const
+{
+	GetAudioFormatSettings().GetWaveFormatModuleHints(OutModuleNames);
+}
+
+#endif // WITH_ENGINE
\ No newline at end of file
diff --git a/Engine/Source/Developer/TargetPlatform/Private/TargetPlatformManagerModule.cpp b/Engine/Source/Developer/TargetPlatform/Private/TargetPlatformManagerModule.cpp
--- a/Engine/Source/Developer/TargetPlatform/Private/TargetPlatformManagerModule.cpp
+++ b/Engine/Source/Developer/TargetPlatform/Private/TargetPlatformManagerModule.cpp
@@ -14,7 +14,11 @@
 #include "Misc/MonitoredProcess.h"
 #include "Modules/ModuleManager.h"
 #include "Interfaces/ITargetPlatform.h"
+#include "Interfaces/ITargetPlatformControls.h"
+#include "Interfaces/ITargetPlatformSettings.h"
 #include "Interfaces/ITargetPlatformModule.h"
+#include "Interfaces/ITargetPlatformControlsModule.h"
+#include "Interfaces/ITargetPlatformSettingsModule.h"
 #include "Interfaces/ITargetPlatformManagerModule.h"
 #include "Interfaces/IAudioFormat.h"
 #include "Interfaces/IAudioFormatModule.h"
@@ -106,8 +110,10 @@ static const size_t MaxPlatformCount = 64;		// In the unlikely event that someon
 												// uint64 per platform.
 
 static const ITargetPlatform* TargetPlatformArray[MaxPlatformCount];
+static const ITargetPlatformControls* TargetPlatformControlsArray[MaxPlatformCount];
 
 static int32 PlatformCounter = 0;
+static int32 PlatformControlsCounter = 0;
 
 int32 ITargetPlatform::AssignPlatformOrdinal(const ITargetPlatform& Platform)
 {
@@ -120,12 +126,32 @@ int32 ITargetPlatform::AssignPlatformOrdinal(const ITargetPlatform& Platform)
 
 	return Ordinal;
 }
+int32 ITargetPlatformControls::AssignPlatformOrdinal(const ITargetPlatformControls& Platform)
+{
+	check(PlatformControlsCounter < MaxPlatformCount);
+
+	const int32 Ordinal = PlatformControlsCounter++;
 
-const ITargetPlatform* ITargetPlatform::GetPlatformFromOrdinal(int32 Ordinal)
+	check(TargetPlatformControlsArray[Ordinal] == nullptr);
+	TargetPlatformControlsArray[Ordinal] = &Platform;
+
+	return Ordinal;
+}
+
+const ITargetPlatformControls* ITargetPlatformControls::GetPlatformFromOrdinal(int32 Ordinal)
 {
-	check(Ordinal < PlatformCounter);
+	check(Ordinal < PlatformControlsCounter);
 
-	return TargetPlatformArray[Ordinal];
+	return TargetPlatformControlsArray[Ordinal];
+}
+
+const class ITargetPlatformSettings& ITargetDevice::GetPlatformSettings() const
+{
+	return GetTargetPlatform().GetPlatformSettings();
+}
+const class ITargetPlatformControls& ITargetDevice::GetPlatformControls() const
+{
+	return GetTargetPlatform().GetPlatformControls();
 }
 
 ITargetPlatform::FOnTargetDeviceDiscovered& ITargetPlatform::OnDeviceDiscovered()
@@ -140,6 +166,17 @@ ITargetPlatform::FOnTargetDeviceLost& ITargetPlatform::OnDeviceLost()
 	return Delegate;
 }
 
+ITargetPlatformControls::FOnTargetDeviceDiscovered& ITargetPlatformControls::OnDeviceDiscovered()
+{
+	static FOnTargetDeviceDiscovered Delegate;
+	return Delegate;
+}
+
+ITargetPlatformControls::FOnTargetDeviceLost& ITargetPlatformControls::OnDeviceLost()
+{
+	static FOnTargetDeviceLost Delegate;
+	return Delegate;
+}
 
 /**
  * Module for the target platform manager
@@ -261,6 +298,26 @@ public:
 		return Platforms;
 	}
 
+	virtual const TArray<ITargetPlatformControls*>& GetTargetPlatformControls() override
+	{
+		if (PlatformControls.Num() == 0 || bForceCacheUpdate)
+		{
+			DiscoverAvailablePlatforms();
+		}
+
+		return PlatformControls;
+	}
+
+	virtual const TArray<ITargetPlatformSettings*>& GetTargetPlatformSettings() override
+	{
+		if (PlatformSettings.Num() == 0 || bForceCacheUpdate)
+		{
+			DiscoverAvailablePlatforms();
+		}
+
+		return PlatformSettings;
+	}
+
 	virtual ITargetDevicePtr FindTargetDevice(const FTargetDeviceId& DeviceId) override
 	{
 		ITargetPlatform* Platform = FindTargetPlatform(DeviceId.GetPlatformName());
@@ -703,8 +760,13 @@ protected:
 
 		// try the incoming name as a module name, or as a platform name
 		FName PlatformModuleName = PlatformName;
+		//@todo Custom TargetPlatforms
+		FName PlatformControlsModuleName = *(PlatformName.ToString() + TEXT("TargetPlatformControls"));;
+		FName PlatformSettingsModuleName = *(PlatformName.ToString() + TEXT("TargetPlatformSettings"));;
 
 		ITargetPlatformModule* Module = nullptr;
+		ITargetPlatformControlsModule* ModuleControls = nullptr;
+		ITargetPlatformSettingsModule* ModuleSettings = nullptr;
 
 		if (!FModuleManager::Get().ModuleExists(*PlatformModuleName.ToString()))
 		{
@@ -716,6 +778,26 @@ protected:
 			Module = FModuleManager::LoadModulePtr<ITargetPlatformModule>(PlatformModuleName);
 		}
 
+		if (FModuleManager::Get().ModuleExists(*PlatformSettingsModuleName.ToString()))
+		{
+			ModuleSettings = FModuleManager::LoadModulePtr<ITargetPlatformSettingsModule>(PlatformSettingsModuleName);
+		}
+
+		if (FModuleManager::Get().ModuleExists(*PlatformControlsModuleName.ToString()))
+		{
+			ModuleControls = FModuleManager::LoadModulePtr<ITargetPlatformControlsModule>(PlatformControlsModuleName);
+		}
+
+		if(ModuleSettings != nullptr)
+		{
+			TArray<ITargetPlatformSettings*> TargetPlatformSettings = ModuleSettings->GetTargetPlatformSettings();
+			for (ITargetPlatformSettings* Platform : TargetPlatformSettings)
+			{
+				PlatformSettings.Add(Platform);
+				Module->PlatformSettings.Add(Platform);
+			}
+		}
+
 		// original logic for module loading here
 		if (Module)
 		{
@@ -726,6 +808,17 @@ protected:
 		RETRY_SETUPANDVALIDATE:
 			if (AutoSDKPath == TEXT("") || SetupAndValidateAutoSDK(AutoSDKPath))
 			{
+
+				if (ModuleControls != nullptr)
+				{
+					TArray<ITargetPlatformControls*> TargetPlatformControls = ModuleControls->GetTargetPlatformControls(PlatformSettingsModuleName);
+					for (ITargetPlatformControls* Platform : TargetPlatformControls)
+					{
+						PlatformControls.Add(Platform);
+						Module->PlatformControls.Add(Platform);
+					}
+				}
+
 				TArray<ITargetPlatform*> TargetPlatforms = Module->GetTargetPlatforms();
 				for (ITargetPlatform* Platform : TargetPlatforms)
 				{
@@ -1264,6 +1357,8 @@ private:
 
 	// Holds the list of discovered platforms.
 	TArray<ITargetPlatform*> Platforms;
+	TArray<ITargetPlatformControls*> PlatformControls;
+	TArray<ITargetPlatformSettings*> PlatformSettings;
 
 	// Map for fast lookup of platforms by name.
 	TMap<FName, ITargetPlatform*> PlatformsByName;
diff --git a/Engine/Source/Developer/TargetPlatform/Private/TargetPlatformSettingsBase.cpp b/Engine/Source/Developer/TargetPlatform/Private/TargetPlatformSettingsBase.cpp
new file mode 100644
--- /dev/null
+++ b/Engine/Source/Developer/TargetPlatform/Private/TargetPlatformSettingsBase.cpp
@@ -0,0 +1,104 @@
+// Copyright Epic Games, Inc. All Rights Reserved.
+
+#include "Common/TargetPlatformSettingsBase.h"
+#include "AnalyticsEventAttribute.h"
+#include "HAL/IConsoleManager.h"
+
+bool FTargetPlatformSettingsBase::UsesForwardShading() const
+{
+	static IConsoleVariable* CVarForwardShading = IConsoleManager::Get().FindConsoleVariable(TEXT("r.ForwardShading"));
+	return CVarForwardShading ? (CVarForwardShading->GetInt() != 0) : false;
+}
+
+bool FTargetPlatformSettingsBase::UsesDBuffer() const
+{
+	static IConsoleVariable* CVar = IConsoleManager::Get().FindConsoleVariable(TEXT("r.DBuffer"));
+	return CVar ? (CVar->GetInt() != 0) : false;
+}
+
+bool FTargetPlatformSettingsBase::UsesBasePassVelocity() const
+{
+	static IConsoleVariable* CVar = IConsoleManager::Get().FindConsoleVariable(TEXT("r.VelocityOutputPass"));
+	return CVar ? (CVar->GetInt() == 1) : false;
+}
+
+bool FTargetPlatformSettingsBase::VelocityEncodeDepth() const
+{
+	return true;
+}
+
+bool FTargetPlatformSettingsBase::UsesSelectiveBasePassOutputs() const
+{
+	static IConsoleVariable* CVar = IConsoleManager::Get().FindConsoleVariable(TEXT("r.SelectiveBasePassOutputs"));
+	return CVar ? (CVar->GetInt() != 0) : false;
+}
+
+bool FTargetPlatformSettingsBase::UsesDistanceFields() const
+{
+	static IConsoleVariable* CVar = IConsoleManager::Get().FindConsoleVariable(TEXT("r.DistanceFields"));
+	return CVar ? (CVar->GetInt() != 0) : false;
+}
+
+bool FTargetPlatformSettingsBase::UsesRayTracing() const
+{
+	static IConsoleVariable* CVar = IConsoleManager::Get().FindConsoleVariable(TEXT("r.RayTracing"));
+	return CVar ? (CVar->GetInt() != 0) : false;
+}
+
+EOfflineBVHMode FTargetPlatformSettingsBase::GetStaticMeshOfflineBVHMode() const
+{
+	return EOfflineBVHMode::Disabled;
+}
+
+bool FTargetPlatformSettingsBase::GetStaticMeshOfflineBVHCompression() const
+{
+	return false;
+}
+
+bool FTargetPlatformSettingsBase::ForcesSimpleSkyDiffuse() const
+{
+	return false;
+}
+
+float FTargetPlatformSettingsBase::GetDownSampleMeshDistanceFieldDivider() const
+{
+	return 1.0f;
+}
+
+int32 FTargetPlatformSettingsBase::GetHeightFogModeForOpaque() const
+{
+	// Don't override the project setting by default
+	// Platforms wish to support override need to implement the logic in their own target platform classes
+	return 0;
+}
+
+bool FTargetPlatformSettingsBase::UsesMobileAmbientOcclusion() const
+{
+	static IConsoleVariable* CVar = IConsoleManager::Get().FindConsoleVariable(TEXT("r.Mobile.AmbientOcclusion"));
+	return CVar ? (CVar->GetInt() != 0) : false;
+}
+
+bool FTargetPlatformSettingsBase::UsesMobileDBuffer() const
+{
+	static IConsoleVariable* CVar = IConsoleManager::Get().FindConsoleVariable(TEXT("r.Mobile.DBuffer"));
+	return CVar ? (CVar->GetInt() != 0) : false;
+}
+
+bool FTargetPlatformSettingsBase::UsesASTCHDR() const
+{
+	static IConsoleVariable* CVar = IConsoleManager::Get().FindConsoleVariable(TEXT("cook.ASTCTextureCompressor"));
+	const bool bUsesARMCompressor = (CVar ? (CVar->GetInt() != 0) : false);
+
+	static IConsoleVariable* CVarASTCHDR = IConsoleManager::Get().FindConsoleVariable(TEXT("cook.AllowASTCHDRProfile"));
+	const bool bUsesASTCHDR = (CVarASTCHDR ? (CVarASTCHDR->GetInt() != 0) : false);
+
+	return (bUsesARMCompressor && bUsesASTCHDR);
+}
+
+void FTargetPlatformSettingsBase::GetRayTracingShaderFormats(TArray<FName>& OutFormats) const
+{
+	if (UsesRayTracing())
+	{
+		GetAllTargetedShaderFormats(OutFormats);
+	}
+}
\ No newline at end of file
diff --git a/Engine/Source/Developer/TargetPlatform/Public/Common/TargetPlatformBase.h b/Engine/Source/Developer/TargetPlatform/Public/Common/TargetPlatformBase.h
--- a/Engine/Source/Developer/TargetPlatform/Public/Common/TargetPlatformBase.h
+++ b/Engine/Source/Developer/TargetPlatform/Public/Common/TargetPlatformBase.h
@@ -611,3 +611,367 @@ protected:
 };
 
 
+// Redirect everything to TPS or TPC
+class FTargetPlatformMerged final
+	: public ITargetPlatform
+{
+public:
+	FTargetPlatformMerged(ITargetPlatformSettings* InTargetPlatformSettings, ITargetPlatformControls* InTargetPlatformControls)
+		: TargetPlatformSettings(InTargetPlatformSettings)
+		, TargetPlatformControls(InTargetPlatformControls)
+	{
+
+	}
+	virtual bool AddDevice(const FString& DeviceName, bool bDefault) override 
+	{ 
+		return TargetPlatformControls->AddDevice(DeviceName, bDefault); 
+	}
+	virtual bool AddDevice(const FString& DeviceId, const FString& DeviceUserFriendlyName, const FString& Username, const FString& Password, bool bDefault) override 
+	{ 
+		return TargetPlatformControls->AddDevice(DeviceId, DeviceUserFriendlyName, Username, Password, bDefault); 
+	}
+	virtual FString PlatformName() const override 
+	{
+		return TargetPlatformControls->PlatformName();
+	}
+	virtual FText DisplayName() const override 
+	{
+		return TargetPlatformControls->DisplayName();
+	}
+	virtual int32 CheckRequirements(bool bProjectHasCode, EBuildConfiguration Configuration, bool bRequiresAssetNativization, FString& OutTutorialPath, FString& OutDocumentationPath, FText& CustomizedLogMessage) const override 
+	{
+		return TargetPlatformControls->CheckRequirements(bProjectHasCode, Configuration, bRequiresAssetNativization, OutTutorialPath, OutDocumentationPath, CustomizedLogMessage);
+	}
+	virtual bool RequiresTempTarget(bool bProjectHasCode, EBuildConfiguration Configuration, bool bRequiresAssetNativization, FText& OutReason) const override 
+	{
+		return TargetPlatformControls->RequiresTempTarget(bProjectHasCode, Configuration, bRequiresAssetNativization, OutReason);
+	}
+	virtual const PlatformInfo::FTargetPlatformInfo& GetTargetPlatformInfo() const override 
+	{
+		return TargetPlatformControls->GetTargetPlatformInfo();
+	}
+	virtual const FDataDrivenPlatformInfo& GetPlatformInfo() const override 
+	{
+		return TargetPlatformControls->GetPlatformInfo();
+	}
+	virtual FConfigCacheIni* GetConfigSystem() const override 
+	{
+		return TargetPlatformSettings->GetConfigSystem();
+	}
+	virtual FString IniPlatformName() const override 
+	{
+		return TargetPlatformSettings->IniPlatformName();
+	}
+	virtual FString CookingDeviceProfileName() const override 
+	{
+		return TargetPlatformControls->CookingDeviceProfileName();
+	}
+	virtual void EnableDeviceCheck(bool OnOff) override 
+	{
+		TargetPlatformControls->EnableDeviceCheck(OnOff);
+	}
+	virtual void GetAllDevices(TArray<ITargetDevicePtr>& OutDevices) const override 
+	{
+		TargetPlatformControls->GetAllDevices(OutDevices);
+	}
+	virtual FName GetZlibReplacementFormat() const override 
+	{
+		return TargetPlatformControls->GetZlibReplacementFormat();
+	}
+	virtual int64 GetMemoryMappingAlignment() const override 
+	{
+		return TargetPlatformControls->GetMemoryMappingAlignment();
+	}
+	virtual bool GenerateStreamingInstallManifest(const TMultiMap<FString, int32>& PakchunkMap, const TSet<int32>& PakchunkIndicesInUse) const override 
+	{
+		return TargetPlatformControls->GenerateStreamingInstallManifest(PakchunkMap, PakchunkIndicesInUse);
+	}
+	virtual ITargetDevicePtr GetDefaultDevice() const override 
+	{
+		return TargetPlatformControls->GetDefaultDevice();
+	}
+	virtual ITargetDevicePtr GetDevice(const FTargetDeviceId& DeviceId) override 
+	{
+		return TargetPlatformControls->GetDevice(DeviceId);
+	}
+	virtual bool HasEditorOnlyData() const override 
+	{
+		return TargetPlatformControls->HasEditorOnlyData();
+	}
+	virtual bool AllowsEditorObjects() const override 
+	{
+		return TargetPlatformControls->AllowsEditorObjects();
+	}
+	virtual bool AllowsDevelopmentObjects() const override 
+	{
+		return TargetPlatformControls->AllowsDevelopmentObjects();
+	}
+	virtual bool IsClientOnly() const override 
+	{
+		return TargetPlatformControls->IsClientOnly();
+	}
+	virtual bool IsLittleEndian() const override 
+	{
+		return TargetPlatformControls->IsLittleEndian();
+	}
+	virtual bool IsRunningPlatform() const override 
+	{
+		return TargetPlatformControls->IsRunningPlatform();
+	}
+	virtual bool IsServerOnly() const override 
+	{
+		return TargetPlatformControls->IsServerOnly();
+	}
+	virtual bool IsEnabledForPlugin(const IPlugin& Plugin) const override 
+	{
+		return TargetPlatformControls->IsEnabledForPlugin(Plugin);
+	}
+	virtual bool CanSupportRemoteShaderCompile() const override 
+	{
+		return TargetPlatformControls->CanSupportRemoteShaderCompile();
+	}
+	virtual void GetShaderCompilerDependencies(TArray<FString>& OutDependencies) const override 
+	{
+		TargetPlatformControls->GetShaderCompilerDependencies(OutDependencies);
+	}
+	virtual bool IsSdkInstalled(bool bProjectHasCode, FString& OutDocumentationPath) const override 
+	{
+		return TargetPlatformControls->IsSdkInstalled(bProjectHasCode, OutDocumentationPath);
+	}
+	virtual bool RequiresCookedData() const override 
+	{
+		return TargetPlatformControls->RequiresCookedData();
+	}
+	virtual bool RequiresOriginalReleaseVersionForPatch() const override 
+	{
+		return TargetPlatformControls->RequiresOriginalReleaseVersionForPatch();
+	}
+	virtual bool HasSecurePackageFormat() const override 
+	{
+		return TargetPlatformControls->HasSecurePackageFormat();
+	}
+	virtual EPlatformAuthentication RequiresUserCredentials() const override 
+	{
+		return TargetPlatformControls->RequiresUserCredentials();
+	}
+	virtual bool SupportsAutoSDK() const override 
+	{
+		return TargetPlatformControls->SupportsAutoSDK();
+	}
+	virtual bool SupportsBuildTarget(EBuildTargetType TargetType) const override 
+	{
+		return TargetPlatformControls->SupportsBuildTarget(TargetType);
+	}
+	virtual EBuildTargetType GetRuntimePlatformType() const override 
+	{
+		return TargetPlatformControls->GetRuntimePlatformType();
+	}
+	virtual bool SupportsFeature(ETargetPlatformFeatures Feature) const override 
+	{
+		return TargetPlatformSettings->SupportsFeature(Feature);
+	}
+	virtual bool SupportsValueForType(FName SupportedType, FName RequiredSupportedValue) const override 
+	{
+		return TargetPlatformSettings->SupportsValueForType(SupportedType, RequiredSupportedValue);
+	}
+	virtual bool UsesForwardShading() const override 
+	{
+		return TargetPlatformSettings->UsesForwardShading();
+	}
+	virtual bool UsesDBuffer() const override 
+	{
+		return TargetPlatformSettings->UsesDBuffer();
+	}
+	virtual bool UsesBasePassVelocity() const override 
+	{
+		return TargetPlatformSettings->UsesBasePassVelocity();
+	}
+	virtual bool UsesSelectiveBasePassOutputs() const override 
+	{
+		return TargetPlatformSettings->UsesSelectiveBasePassOutputs();
+	}
+	virtual bool UsesDistanceFields() const override 
+	{
+		return TargetPlatformSettings->UsesDistanceFields();
+	}
+	virtual bool UsesRayTracing() const override 
+	{
+		return TargetPlatformSettings->UsesRayTracing();
+	}
+	virtual EOfflineBVHMode GetStaticMeshOfflineBVHMode() const override 
+	{
+		return TargetPlatformSettings->GetStaticMeshOfflineBVHMode();
+	}
+	virtual bool GetStaticMeshOfflineBVHCompression() const override 
+	{
+		return TargetPlatformSettings->GetStaticMeshOfflineBVHCompression();
+	}
+	virtual bool ForcesSimpleSkyDiffuse() const override 
+	{
+		return TargetPlatformSettings->ForcesSimpleSkyDiffuse();
+	}
+	virtual bool VelocityEncodeDepth() const override 
+	{
+		return TargetPlatformSettings->VelocityEncodeDepth();
+	}
+	virtual float GetDownSampleMeshDistanceFieldDivider() const override 
+	{
+		return TargetPlatformSettings->GetDownSampleMeshDistanceFieldDivider();
+	}
+	virtual int32 GetHeightFogModeForOpaque() const override 
+	{
+		return TargetPlatformSettings->GetHeightFogModeForOpaque();
+	}
+	virtual bool UsesMobileAmbientOcclusion() const override 
+	{
+		return TargetPlatformSettings->UsesMobileAmbientOcclusion();
+	}
+	virtual bool UsesMobileDBuffer() const override 
+	{
+		return TargetPlatformSettings->UsesMobileDBuffer();
+	}
+	virtual bool UsesASTCHDR() const override 
+	{
+		return TargetPlatformSettings->UsesASTCHDR();
+	}
+	virtual void GetAllPossibleShaderFormats(TArray<FName>& OutFormats) const override 
+	{
+		TargetPlatformSettings->GetAllPossibleShaderFormats(OutFormats);
+	}
+	virtual void GetAllTargetedShaderFormats(TArray<FName>& OutFormats) const override 
+	{
+		TargetPlatformSettings->GetAllTargetedShaderFormats(OutFormats);
+	}
+	virtual void GetRayTracingShaderFormats(TArray<FName>& OutFormats) const override 
+	{
+		TargetPlatformSettings->GetRayTracingShaderFormats(OutFormats);
+	}
+	virtual void GetPlatformSpecificProjectAnalytics(TArray<struct FAnalyticsEventAttribute>& AnalyticsParamArray) const override 
+	{
+		TargetPlatformControls->GetPlatformSpecificProjectAnalytics(AnalyticsParamArray);
+	}
+#if WITH_ENGINE
+	virtual FName GetPhysicsFormat(class UBodySetup* Body) const override 
+	{
+		return TargetPlatformControls->GetPhysicsFormat(Body);
+	}
+	virtual void GetReflectionCaptureFormats(TArray<FName>& OutFormats) const override 
+	{
+		TargetPlatformSettings->GetReflectionCaptureFormats(OutFormats);
+	}
+	virtual void GetShaderFormatModuleHints(TArray<FName>& OutModuleNames) const override 
+	{
+		TargetPlatformControls->GetShaderFormatModuleHints(OutModuleNames);
+	}
+	virtual void GetTextureFormats(const class UTexture* Texture, TArray< TArray<FName> >& OutFormats) const override 
+	{
+		TargetPlatformControls->GetTextureFormats(Texture, OutFormats);
+	}
+	virtual void GetAllTextureFormats(TArray<FName>& OutFormats) const override 
+	{
+		TargetPlatformControls->GetAllTextureFormats(OutFormats);
+	}
+
+	virtual void GetTextureFormatModuleHints(TArray<FName>& OutModuleNames) const override 
+	{
+		TargetPlatformControls->GetTextureFormatModuleHints(OutModuleNames);
+	}
+	virtual FName FinalizeVirtualTextureLayerFormat(FName Format) const override 
+	{
+		return TargetPlatformControls->FinalizeVirtualTextureLayerFormat(Format);
+	}
+	virtual bool SupportsLQCompressionTextureFormat() const override 
+	{
+		return TargetPlatformControls->SupportsLQCompressionTextureFormat();
+	}
+	virtual FName GetWaveFormat(const class USoundWave* Wave) const override 
+	{
+		return TargetPlatformControls->GetWaveFormat(Wave);
+	}
+	virtual void GetAllWaveFormats(TArray<FName>& OutFormats) const override 
+	{
+		TargetPlatformControls->GetAllWaveFormats(OutFormats);
+	}
+	virtual void GetWaveFormatModuleHints(TArray<FName>& OutModuleNames) const override 
+	{
+		TargetPlatformControls->GetWaveFormatModuleHints(OutModuleNames);
+	}
+	virtual bool AllowAudioVisualData() const override 
+	{
+		return TargetPlatformControls->AllowAudioVisualData();
+	}
+	virtual bool AllowObject(const class UObject* Object) const override 
+	{
+		return TargetPlatformControls->AllowObject(Object);
+	}
+	virtual const class UTextureLODSettings& GetTextureLODSettings() const override 
+	{
+		return TargetPlatformSettings->GetTextureLODSettings();
+	}
+	virtual void RegisterTextureLODSettings(const class UTextureLODSettings* InTextureLODSettings) override 
+	{
+		TargetPlatformSettings->RegisterTextureLODSettings(InTextureLODSettings);
+	}
+	virtual const class FStaticMeshLODSettings& GetStaticMeshLODSettings() const override 
+	{
+		return TargetPlatformSettings->GetStaticMeshLODSettings();
+	}
+	virtual FName GetMeshBuilderModuleName() const override 
+	{
+		return TargetPlatformControls->GetMeshBuilderModuleName();
+	}
+#endif
+	virtual bool PackageBuild(const FString& InPackgeDirectory) override 
+	{
+		return TargetPlatformControls->PackageBuild(InPackgeDirectory);
+	}
+	virtual bool SupportsVariants() const override 
+	{
+		return TargetPlatformControls->SupportsVariants();
+	}
+	virtual float GetVariantPriority() const override 
+	{
+		return TargetPlatformControls->SupportsVariants();
+	}
+	virtual bool SendLowerCaseFilePaths() const override 
+	{
+		return TargetPlatformControls->SupportsVariants();
+	}
+	virtual void GetBuildProjectSettingKeys(FString& OutSection, TArray<FString>& InBoolKeys, TArray<FString>& InIntKeys, TArray<FString>& InStringKeys) const override 
+	{
+		TargetPlatformControls->GetBuildProjectSettingKeys(OutSection, InBoolKeys, InIntKeys, InStringKeys);
+	}
+	virtual int32 GetPlatformOrdinal() const override 
+	{
+		return TargetPlatformControls->GetPlatformOrdinal();
+	}
+	virtual TSharedPtr<IDeviceManagerCustomPlatformWidgetCreator> GetCustomWidgetCreator() const override 
+	{
+		return TargetPlatformControls->GetCustomWidgetCreator();
+	}
+	virtual bool ShouldExpandTo32Bit(const uint16* Indices, const int32 NumIndices) const override 
+	{
+		return TargetPlatformControls->ShouldExpandTo32Bit(Indices, NumIndices);
+	}
+	virtual bool CopyFileToTarget(const FString& DeviceId, const FString& HostFilename, const FString& TargetFilename, const TMap<FString, FString>& CustomPlatformData) override 
+	{
+		return TargetPlatformControls->CopyFileToTarget(DeviceId, HostFilename, TargetFilename, CustomPlatformData);
+	}
+	virtual void GetExtraPackagesToCook(TArray<FName>& PackageNames) const override
+	{
+		TargetPlatformControls->GetExtraPackagesToCook(PackageNames);
+	}
+	virtual bool InitializeHostPlatform() override 
+	{
+		return TargetPlatformControls->InitializeHostPlatform();
+	}
+
+public:
+
+	/** Virtual destructor. */
+	virtual ~FTargetPlatformMerged() {}
+
+private:
+	ITargetPlatformSettings* TargetPlatformSettings;
+	ITargetPlatformControls* TargetPlatformControls;
+};
\ No newline at end of file
diff --git a/Engine/Source/Developer/TargetPlatform/Public/Common/TargetPlatformControlsBase.h b/Engine/Source/Developer/TargetPlatform/Public/Common/TargetPlatformControlsBase.h
new file mode 100644
--- /dev/null
+++ b/Engine/Source/Developer/TargetPlatform/Public/Common/TargetPlatformControlsBase.h
@@ -0,0 +1,402 @@
+// Copyright Epic Games, Inc. All Rights Reserved.
+
+#pragma once
+
+#include "Interfaces/ITargetPlatformControls.h"
+#include "GenericPlatform/GenericPlatformFile.h"
+#include "HAL/PlatformFileManager.h"
+#include "CoreGlobals.h"
+#include "PlatformInfo.h"
+#include "Interfaces/ITargetDevice.h"
+#include "Interfaces/ITargetPlatformSettings.h"
+#include "Templates/PimplPtr.h"
+#include "Misc/Paths.h"
+#include "Misc/ConfigCacheIni.h"
+// Forward declare.
+namespace Audio
+{
+	class FAudioFormatSettings;
+}
+
+class FTargetPlatformControlsBase
+	: public ITargetPlatformControls
+{
+public:
+	virtual bool AddDevice(const FString& DeviceName, bool bDefault) override
+	{
+		return false;
+	}
+
+	virtual bool AddDevice(const FString& DeviceId, const FString& DeviceUserFriendlyName, const FString& Username, const FString& Password, bool bDefault) override
+	{
+		return AddDevice(DeviceId, bDefault);
+	}
+
+	virtual FText DisplayName() const override
+	{
+		return PlatformInfo->DisplayName;
+	}
+
+	virtual const PlatformInfo::FTargetPlatformInfo& GetTargetPlatformInfo() const override
+	{
+		return *PlatformInfo;
+	}
+
+	virtual const FDataDrivenPlatformInfo& GetPlatformInfo() const override
+	{
+		return *PlatformInfo->DataDrivenPlatformInfo;
+	}
+
+	TARGETPLATFORM_API virtual bool IsEnabledForPlugin(const IPlugin& Plugin) const override;
+	TARGETPLATFORM_API virtual void GetPlatformSpecificProjectAnalytics(TArray<struct FAnalyticsEventAttribute>& AnalyticsParamArray) const override;
+
+#if WITH_ENGINE
+	virtual FName FinalizeVirtualTextureLayerFormat(FName Format) const override
+	{
+		return Format;
+	}
+	virtual bool SupportsLQCompressionTextureFormat() const override { return true; };
+#endif
+
+	virtual bool PackageBuild(const FString& InPackgeDirectory) override
+	{
+		return true;
+	}
+
+	virtual bool CanSupportRemoteShaderCompile() const override
+	{
+		return true;
+	}
+
+	virtual void GetShaderCompilerDependencies(TArray<FString>& OutDependencies) const override
+	{
+	}
+
+	/** Helper method to fill a dependencies array for the shader compiler with absolute paths, passing a relative path to the Engine as the parameter. */
+	static void AddDependencySCArrayHelper(TArray<FString>& OutDependencies, const FString& DependencyRelativePath)
+	{
+		IPlatformFile& PlatformFile = FPlatformFileManager::Get().GetPlatformFile();
+		FString DependencyAbsolutePath = PlatformFile.ConvertToAbsolutePathForExternalAppForRead(*(FPaths::EngineDir() / DependencyRelativePath));
+		FPaths::NormalizeDirectoryName(DependencyAbsolutePath);
+		OutDependencies.AddUnique(DependencyAbsolutePath);
+	}
+
+	virtual bool IsSdkInstalled(bool bProjectHasCode, FString& OutDocumentationPath) const override
+	{
+		return true;
+	}
+
+	virtual int32 CheckRequirements(bool bProjectHasCode, EBuildConfiguration Configuration, bool bRequiresAssetNativization, FString& OutTutorialPath, FString& OutDocumentationPath, FText& CustomizedLogMessage) const override
+	{
+		int32 bReadyToBuild = ETargetPlatformReadyStatus::Ready; // @todo How do we check that the iOS SDK is installed when building from Windows? Is that even possible?
+		if (!IsSdkInstalled(bProjectHasCode, OutTutorialPath))
+		{
+			bReadyToBuild |= ETargetPlatformReadyStatus::SDKNotFound;
+		}
+		return bReadyToBuild;
+	}
+
+	TARGETPLATFORM_API virtual bool RequiresTempTarget(bool bProjectHasCode, EBuildConfiguration Configuration, bool bRequiresAssetNativization, FText& OutReason) const override;
+
+	virtual bool SupportsVariants() const override
+	{
+		return false;
+	}
+
+	virtual float GetVariantPriority() const override
+	{
+		return IsClientOnly() ? 0.0f : 0.2f;
+	}
+
+	virtual bool SendLowerCaseFilePaths() const override
+	{
+		return false;
+	}
+
+	virtual bool AllowsEditorObjects() const override
+	{
+		return HasEditorOnlyData();
+	}
+
+	virtual bool AllowsDevelopmentObjects() const override
+	{
+		return AllowsEditorObjects();
+	}
+
+	virtual void GetBuildProjectSettingKeys(FString& OutSection, TArray<FString>& InBoolKeys, TArray<FString>& InIntKeys, TArray<FString>& InStringKeys) const override
+	{
+		// do nothing in the base class
+	}
+
+	virtual int32 GetPlatformOrdinal() const override
+	{
+		return PlatformOrdinal;
+	}
+
+	TARGETPLATFORM_API virtual TSharedPtr<IDeviceManagerCustomPlatformWidgetCreator> GetCustomWidgetCreator() const override;
+
+	virtual bool ShouldExpandTo32Bit(const uint16* Indices, const int32 NumIndices) const override
+	{
+		return false;
+	}
+
+#if WITH_ENGINE
+	virtual bool AllowAudioVisualData() const override
+	{
+		return !IsServerOnly();
+	}
+
+	virtual bool AllowObject(const class UObject* Object) const override
+	{
+		return true;
+	}
+
+	virtual FName GetMeshBuilderModuleName() const override
+	{
+		// MeshBuilder is the default module. Platforms may override this to provide platform specific mesh data.
+		static const FName NAME_MeshBuilder(TEXT("MeshBuilder"));
+		return NAME_MeshBuilder;
+	}
+
+	virtual void GetShaderFormatModuleHints(TArray<FName>& OutModuleNames) const override
+	{
+	}
+
+	virtual void GetTextureFormatModuleHints(TArray<FName>& OutModuleNames) const override
+	{
+		// these are the default texture format modules, since many platforms 
+		OutModuleNames.Add(TEXT("TextureFormatUncompressed"));
+		OutModuleNames.Add(TEXT("TextureFormatDXT"));
+		OutModuleNames.Add(TEXT("TextureFormatIntelISPCTexComp"));
+
+		// there is a possible optional format module name in the ini (alternate texture compressor)
+		FString TextureCompressionFormat;
+		if (TargetPlatformSettings->GetConfigSystem()->GetString(TEXT("AlternateTextureCompression"), TEXT("TextureCompressionFormat"), TextureCompressionFormat, GEngineIni))
+		{
+			OutModuleNames.Add(*TextureCompressionFormat);
+		}
+	}
+
+	TARGETPLATFORM_API virtual FName GetWaveFormat(const class USoundWave* Wave) const override;
+
+	TARGETPLATFORM_API virtual void GetAllWaveFormats(TArray<FName>& OutFormats) const override;
+
+	TARGETPLATFORM_API virtual void GetWaveFormatModuleHints(TArray<FName>& OutModuleNames) const override;
+
+#endif
+
+	virtual bool CopyFileToTarget(const FString& TargetAddress, const FString& HostFilename, const FString& TargetFilename, const TMap<FString, FString>& CustomPlatformData) override
+	{
+		return false;
+	}
+
+	virtual void GetExtraPackagesToCook(TArray<FName>& PackageNames) const override
+	{
+
+	}
+
+	virtual bool InitializeHostPlatform()
+	{
+		// if the platform doesn't need anything, it's valid to do nothing
+		return true;
+	}
+
+protected:
+	TARGETPLATFORM_API FTargetPlatformControlsBase(const PlatformInfo::FTargetPlatformInfo* const InPlatformInfo, ITargetPlatformSettings* TargetPlatformSettings);
+
+	/** Information about this platform */
+	const PlatformInfo::FTargetPlatformInfo* PlatformInfo;
+	int32 PlatformOrdinal;
+
+	TARGETPLATFORM_API const Audio::FAudioFormatSettings& GetAudioFormatSettings() const;
+
+	/** Analytics helper functions */
+	TARGETPLATFORM_API static void AppendAnalyticsEventConfigBool(TArray<struct FAnalyticsEventAttribute>& AnalyticsParamArray, const TCHAR* ConfigSection, const TCHAR* ConfigKey, const FString& IniFileName, const TCHAR* AnalyticsKeyNameOverride = nullptr);
+	TARGETPLATFORM_API static void AppendAnalyticsEventConfigInt(TArray<struct FAnalyticsEventAttribute>& AnalyticsParamArray, const TCHAR* ConfigSection, const TCHAR* ConfigKey, const FString& IniFileName, const TCHAR* AnalyticsKeyNameOverride = nullptr);
+	TARGETPLATFORM_API static void AppendAnalyticsEventConfigFloat(TArray<struct FAnalyticsEventAttribute>& AnalyticsParamArray, const TCHAR* ConfigSection, const TCHAR* ConfigKey, const FString& IniFileName, const TCHAR* AnalyticsKeyNameOverride = nullptr);
+	TARGETPLATFORM_API static void AppendAnalyticsEventConfigString(TArray<struct FAnalyticsEventAttribute>& AnalyticsParamArray, const TCHAR* ConfigSection, const TCHAR* ConfigKey, const FString& IniFileName, const TCHAR* AnalyticsKeyNameOverride = nullptr);
+	TARGETPLATFORM_API static void AppendAnalyticsEventConfigArray(TArray<struct FAnalyticsEventAttribute>& AnalyticsParamArray, const TCHAR* ConfigSection, const TCHAR* ConfigKey, const FString& IniFileName, const TCHAR* AnalyticsKeyNameOverride = nullptr);
+
+private:
+	TPimplPtr<Audio::FAudioFormatSettings> AudioFormatSettings;
+
+	bool HasDefaultBuildSettings() const;
+	static bool DoProjectSettingsMatchDefault(const FString& InPlatformName, const FString& InSection, const TArray<FString>* InBoolKeys, const TArray<FString>* InIntKeys, const TArray<FString>* InStringKeys);
+	static bool IsPluginEnabledForTarget(const IPlugin& Plugin, const struct FProjectDescriptor* Project, const FString& Platform, EBuildConfiguration Configuration, EBuildTargetType TargetType);
+};
+
+template<typename TPlatformProperties>
+class TTargetPlatformControlsBase
+	: public FTargetPlatformControlsBase
+{
+public:
+	/**
+	 * Returns true if the target platform will be able to be  initialized with an FPlatformInfo. Because FPlatformInfo now comes from a .ini file,
+	 * it's possible that the .dll exists, but the .ini does not (should be uncommon, but is necessary to be handled)
+	 */
+	static bool IsUsable()
+	{
+		return true;
+	}
+
+	/**
+	 * Constructor that already has a TPI (notably coming from TNonDesktopTargetPlatform)
+	 */
+	TTargetPlatformControlsBase(const PlatformInfo::FTargetPlatformInfo* const InPlatformInfo, ITargetPlatformSettings* TargetPlatformSettings)
+		: FTargetPlatformControlsBase(InPlatformInfo, TargetPlatformSettings)
+	{
+		// HasEditorOnlyData and RequiresCookedData are mutually exclusive.
+		check(TPlatformProperties::HasEditorOnlyData() != TPlatformProperties::RequiresCookedData());
+	}
+
+	virtual bool HasEditorOnlyData() const override
+	{
+		return TPlatformProperties::HasEditorOnlyData();
+	}
+
+	virtual bool IsLittleEndian() const override
+	{
+		return TPlatformProperties::IsLittleEndian();
+	}
+
+	virtual bool IsServerOnly() const override
+	{
+		return TPlatformProperties::IsServerOnly();
+	}
+
+	virtual bool IsClientOnly() const override
+	{
+		return TPlatformProperties::IsClientOnly();
+	}
+
+	virtual FString PlatformName() const override
+	{
+		// we assume these match for DesktopPlatforms (NonDesktop doesn't return "FooClient", but Desktop does, for legacy reasons)
+		checkSlow(this->PlatformInfo->Name == TPlatformProperties::PlatformName());
+
+		return FString(TPlatformProperties::PlatformName());
+	}
+
+	virtual FString CookingDeviceProfileName() const override
+	{
+		// default to using TargetPlatform name, like WindowsClient, when looking up the DeviceProfile to cook with
+		return PlatformName();
+	}
+
+	virtual bool RequiresCookedData() const override
+	{
+		return TPlatformProperties::RequiresCookedData();
+	}
+
+	virtual bool RequiresOriginalReleaseVersionForPatch() const override
+	{
+		return TPlatformProperties::RequiresOriginalReleaseVersionForPatch();
+	}
+
+	virtual bool HasSecurePackageFormat() const override
+	{
+		return TPlatformProperties::HasSecurePackageFormat();
+	}
+
+	virtual EPlatformAuthentication RequiresUserCredentials() const override
+	{
+		if (TPlatformProperties::RequiresUserCredentials())
+		{
+			return EPlatformAuthentication::Always;
+		}
+		else
+		{
+			return EPlatformAuthentication::Never;
+		}
+	}
+
+	virtual bool SupportsBuildTarget(EBuildTargetType TargetType) const override
+	{
+		return TPlatformProperties::SupportsBuildTarget(TargetType);
+	}
+
+	virtual EBuildTargetType GetRuntimePlatformType() const override
+	{
+		if (AllowsEditorObjects())
+		{
+			// Platforms that AllowsEditorObjects need the runtime type Editor to use those objects
+			return EBuildTargetType::Editor;
+		}
+		return PlatformInfo->PlatformType;
+	}
+
+	virtual bool SupportsAutoSDK() const override
+	{
+		return TPlatformProperties::SupportsAutoSDK();
+	}
+
+	virtual FName GetZlibReplacementFormat() const override
+	{
+		return TPlatformProperties::GetZlibReplacementFormat() != nullptr ? FName(TPlatformProperties::GetZlibReplacementFormat()) : NAME_Zlib;
+	}
+
+	virtual int64 GetMemoryMappingAlignment() const override
+	{
+		return TPlatformProperties::GetMemoryMappingAlignment();
+	}
+
+#if WITH_ENGINE
+	virtual FName GetPhysicsFormat(class UBodySetup* Body) const override
+	{
+		return FName(TPlatformProperties::GetPhysicsFormat());
+	}
+#endif // WITH_ENGINE
+};
+
+template<typename TPlatformProperties>
+class TNonDesktopTargetPlatformControlsBase
+	: public TTargetPlatformControlsBase<TPlatformProperties>
+{
+public:
+	TNonDesktopTargetPlatformControlsBase(bool bInIsClientOnly, ITargetPlatformSettings* InTargetPlatformSettings, const TCHAR* CookFlavor = nullptr, const TCHAR* OverrideIniPlatformName = nullptr)
+		: TTargetPlatformControlsBase<TPlatformProperties>(new PlatformInfo::FTargetPlatformInfo(
+			OverrideIniPlatformName ? FString(OverrideIniPlatformName) : FString(TPlatformProperties::IniPlatformName()),
+			bInIsClientOnly ? EBuildTargetType::Client : EBuildTargetType::Game,
+			CookFlavor), InTargetPlatformSettings)
+		,	bIsClientOnly(bInIsClientOnly)
+	{
+
+	}
+	virtual FString PlatformName() const override
+	{
+		// instead of TPlatformProperties (which won't have Client for non-desktop platforms), use the Info's name, which is programmatically made
+		return this->PlatformInfo->Name.ToString();
+	}
+
+	virtual FString CookingDeviceProfileName() const override
+	{
+		// when cooking for non-desktop platforms, always use the base platform name as the DP to cook with (there aren't IOSClient DPs, there are IOS (base), IPhone7, iPad4, etc DPs)
+		return  this->PlatformInfo->IniPlatformName.ToString();
+	}
+
+	virtual bool HasEditorOnlyData() const override
+	{
+		return false;
+	}
+
+	virtual bool IsServerOnly() const override
+	{
+		return false;
+	}
+
+	virtual bool IsClientOnly() const override
+	{
+		return bIsClientOnly;
+	}
+
+	virtual bool IsRunningPlatform() const override
+	{
+		// IsRunningPlatform is only for editor platforms
+
+		return false;
+	}
+
+protected:
+	// true if this target platform is client-only, ie strips out server stuff
+	bool bIsClientOnly;
+};
\ No newline at end of file
diff --git a/Engine/Source/Developer/TargetPlatform/Public/Common/TargetPlatformSettingsBase.h b/Engine/Source/Developer/TargetPlatform/Public/Common/TargetPlatformSettingsBase.h
new file mode 100644
--- /dev/null
+++ b/Engine/Source/Developer/TargetPlatform/Public/Common/TargetPlatformSettingsBase.h
@@ -0,0 +1,153 @@
+// Copyright Epic Games, Inc. All Rights Reserved.
+
+#pragma once
+
+#include "Interfaces/ITargetPlatformSettings.h"
+#include "Misc/ConfigCacheIni.h"
+// Forward declare.
+namespace Audio
+{
+	class FAudioFormatSettings;
+}
+
+class FTargetPlatformSettingsBase
+	: public ITargetPlatformSettings
+{
+public:
+	virtual FConfigCacheIni* GetConfigSystem() const override
+	{
+		return FConfigCacheIni::ForPlatform(*IniPlatformName());
+	}
+
+	TARGETPLATFORM_API virtual bool UsesForwardShading() const override;
+	TARGETPLATFORM_API virtual bool UsesDBuffer() const override;
+	UE_DEPRECATED(5.1, "Use IsUsingBasePassVelocity(EShaderPlatform Platform) in renderutils that will uses FShaderPlatformCachedIniValue to retrieve the cvar value per platform.")
+	TARGETPLATFORM_API virtual bool UsesBasePassVelocity() const override;
+	TARGETPLATFORM_API virtual bool VelocityEncodeDepth() const override;
+	TARGETPLATFORM_API virtual bool UsesSelectiveBasePassOutputs() const override;
+	TARGETPLATFORM_API virtual bool UsesDistanceFields() const override;
+	TARGETPLATFORM_API virtual bool UsesRayTracing() const override;
+	TARGETPLATFORM_API virtual EOfflineBVHMode GetStaticMeshOfflineBVHMode() const override;
+	TARGETPLATFORM_API virtual bool GetStaticMeshOfflineBVHCompression() const override;
+	TARGETPLATFORM_API virtual bool ForcesSimpleSkyDiffuse() const override;
+	TARGETPLATFORM_API virtual float GetDownSampleMeshDistanceFieldDivider() const override;
+	TARGETPLATFORM_API virtual int32 GetHeightFogModeForOpaque() const override;
+	TARGETPLATFORM_API virtual bool UsesMobileAmbientOcclusion() const override;
+	TARGETPLATFORM_API virtual bool UsesMobileDBuffer() const override;
+	TARGETPLATFORM_API virtual bool UsesASTCHDR() const override;
+	TARGETPLATFORM_API virtual void GetRayTracingShaderFormats(TArray<FName>& OutFormats) const override;
+
+#if WITH_ENGINE
+	virtual void GetReflectionCaptureFormats(TArray<FName>& OutFormats) const override
+	{
+		OutFormats.Add(FName(TEXT("FullHDR")));
+	}
+#endif
+
+	virtual bool SupportsValueForType(FName SupportedType, FName RequiredSupportedValue) const override
+	{
+		// check if the given shader format is returned by this TargetPlatform
+		if (SupportedType == TEXT("ShaderFormat"))
+		{
+			TArray<FName> AllPossibleShaderFormats;
+			GetAllPossibleShaderFormats(AllPossibleShaderFormats);
+			return AllPossibleShaderFormats.Contains(RequiredSupportedValue);
+		}
+
+		return false;
+	}
+
+protected:
+	TARGETPLATFORM_API FTargetPlatformSettingsBase() {};
+};
+
+template<typename TPlatformProperties>
+class TTargetPlatformSettingsBase
+	: public FTargetPlatformSettingsBase
+{
+public:
+	TTargetPlatformSettingsBase(const TCHAR* CookFlavor = nullptr, const TCHAR* OverrideIniPlatformName = nullptr)
+		: FTargetPlatformSettingsBase()
+	{
+		IniPlatformNameValue = OverrideIniPlatformName != nullptr ? FString(OverrideIniPlatformName): FString(TPlatformProperties::IniPlatformName());
+	}
+
+	virtual FString IniPlatformName() const override
+	{
+		return IniPlatformNameValue;
+	}
+
+	virtual bool SupportsFeature(ETargetPlatformFeatures Feature) const override
+	{
+		switch (Feature)
+		{
+		case ETargetPlatformFeatures::AudioStreaming:
+			return TPlatformProperties::SupportsAudioStreaming();
+
+		case ETargetPlatformFeatures::DistanceFieldShadows:
+			return TPlatformProperties::SupportsDistanceFieldShadows();
+
+		case ETargetPlatformFeatures::DistanceFieldAO:
+			return TPlatformProperties::SupportsDistanceFieldAO();
+
+		case ETargetPlatformFeatures::GrayscaleSRGB:
+			return TPlatformProperties::SupportsGrayscaleSRGB();
+
+		case ETargetPlatformFeatures::HighQualityLightmaps:
+			return TPlatformProperties::SupportsHighQualityLightmaps();
+
+		case ETargetPlatformFeatures::LowQualityLightmaps:
+			return TPlatformProperties::SupportsLowQualityLightmaps();
+
+		case ETargetPlatformFeatures::MultipleGameInstances:
+			return TPlatformProperties::SupportsMultipleGameInstances();
+
+		case ETargetPlatformFeatures::Packaging:
+			return false;
+
+		case ETargetPlatformFeatures::CanCookPackages:
+			return false;
+
+		case ETargetPlatformFeatures::TextureStreaming:
+			return TPlatformProperties::SupportsTextureStreaming();
+		case ETargetPlatformFeatures::MeshLODStreaming:
+			return TPlatformProperties::SupportsMeshLODStreaming();
+
+		case ETargetPlatformFeatures::MemoryMappedFiles:
+			return TPlatformProperties::SupportsMemoryMappedFiles();
+
+		case ETargetPlatformFeatures::MemoryMappedAudio:
+			return TPlatformProperties::SupportsMemoryMappedAudio();
+		case ETargetPlatformFeatures::MemoryMappedAnimation:
+			return TPlatformProperties::SupportsMemoryMappedAnimation();
+
+		case ETargetPlatformFeatures::VirtualTextureStreaming:
+			return TPlatformProperties::SupportsVirtualTextureStreaming();
+
+		case ETargetPlatformFeatures::SdkConnectDisconnect:
+		case ETargetPlatformFeatures::UserCredentials:
+			break;
+
+		case ETargetPlatformFeatures::MobileRendering:
+			return false;
+		case ETargetPlatformFeatures::DeferredRendering:
+			return true;
+
+		case ETargetPlatformFeatures::ShouldSplitPaksIntoSmallerSizes:
+			return false;
+
+		case ETargetPlatformFeatures::HalfFloatVertexFormat:
+			return true;
+
+		case ETargetPlatformFeatures::LumenGI:
+			return TPlatformProperties::SupportsLumenGI();
+
+		case ETargetPlatformFeatures::HardwareLZDecompression:
+			return TPlatformProperties::SupportsHardwareLZDecompression();
+		}
+
+		return false;
+	}
+private:
+	FString IniPlatformNameValue;
+};
\ No newline at end of file
diff --git a/Engine/Source/Developer/TargetPlatform/Public/Interfaces/ITargetDevice.h b/Engine/Source/Developer/TargetPlatform/Public/Interfaces/ITargetDevice.h
--- a/Engine/Source/Developer/TargetPlatform/Public/Interfaces/ITargetDevice.h
+++ b/Engine/Source/Developer/TargetPlatform/Public/Interfaces/ITargetDevice.h
@@ -8,6 +8,7 @@
 
 class ITargetDevice;
 class ITargetDeviceOutput;
+class ITargetPlatform;
 
 /**
  * Enumerates features that may be supported by target devices.
@@ -326,7 +327,17 @@ public:
 	/**
 	 * Gets the TargetPlatform that this device belongs to.
 	 */
-	virtual const class ITargetPlatform& GetTargetPlatform() const = 0;
+protected:
+	// Temporary until we get rid of this.
+	virtual const class ITargetPlatform& GetTargetPlatform() const
+	{
+		ITargetPlatform* DummyReference = nullptr;
+		return *DummyReference;
+	};
+public:
+
+	TARGETPLATFORM_API const class ITargetPlatformSettings& GetPlatformSettings() const;
+	TARGETPLATFORM_API const class ITargetPlatformControls& GetPlatformControls() const;
 
 	/**
 	 * Checks whether this device is connected.
diff --git a/Engine/Source/Developer/TargetPlatform/Public/Interfaces/ITargetPlatform.h b/Engine/Source/Developer/TargetPlatform/Public/Interfaces/ITargetPlatform.h
--- a/Engine/Source/Developer/TargetPlatform/Public/Interfaces/ITargetPlatform.h
+++ b/Engine/Source/Developer/TargetPlatform/Public/Interfaces/ITargetPlatform.h
@@ -13,6 +13,8 @@
 #include "Internationalization/Text.h"
 #include "Templates/SharedPointer.h"
 #include "UObject/NameTypes.h"
+#include "Interfaces/ITargetPlatformSettings.h"
+#include "Interfaces/ITargetPlatformControls.h"
 
 class FConfigCacheIni;
 class IDeviceManagerCustomPlatformWidgetCreator;
@@ -25,747 +27,25 @@ namespace PlatformInfo
 	struct FTargetPlatformInfo;
 }
 
-/**
- * Enumerates features that may be supported by target platforms.
- */
-enum class ETargetPlatformFeatures
-{
-	/** Audio Streaming */
-	AudioStreaming,
-
-	/** Distance field shadows. */
-	DistanceFieldShadows,
-
-	/** Distance field AO. */
-	DistanceFieldAO,
-
-	/** Gray scale SRGB texture formats support. */
-	GrayscaleSRGB,
-
-	/** High quality light maps. */
-	HighQualityLightmaps,
-
-	/** Low quality light maps. */
-	LowQualityLightmaps,
-
-	/** Run multiple game instances on a single device. */
-	MultipleGameInstances,
-
-	/** Builds can be packaged for this platform. */
-	Packaging,
-
-	/** Connect and disconnect devices through the SDK. */
-	SdkConnectDisconnect,
-
-	/** Texture streaming. */
-	TextureStreaming,
-
-	/** Mesh LOD streaming. */
-	MeshLODStreaming,
-
-	/** Landscape visual mesh LOD streaming. */
-	LandscapeMeshLODStreaming UE_DEPRECATED(5.1, "LandscapeMeshLODStreaming is now deprecated and will be removed."),
-
-	/** User credentials are required to use the device. */
-	UserCredentials,
-
-	/** The platform uses the mobile forward pipeline */
-	MobileRendering,
-
-	/** The platform uses the deferred pipeline, typically PC/Console platforms */
-	DeferredRendering,
-
-	/* Should split paks into smaller sized paks */
-	ShouldSplitPaksIntoSmallerSizes,
-
-	/* The platform supports half float vertex format */
-	HalfFloatVertexFormat,
-
-	/* The platform supports the experimental Device Output Log window */
-	DeviceOutputLog,
-
-	/* The platform supports memory mapped files */
-	MemoryMappedFiles,
-
-	/* The platform supports memory mapped audio */
-	MemoryMappedAudio,
-
-	/* The platform supports memory mapped animation */
-	MemoryMappedAnimation,
-
-	/* The platform supports sparse textures */
-	SparseTextures,
-	
-	/* Can we use the virtual texture streaming system on this platform. */
-	VirtualTextureStreaming,
-
-	/** Lumen Global Illumination. */
-	LumenGI,
-
-	/** The platform supports hardware LZ decompression */
-	HardwareLZDecompression,
-
-	/* The platform makes use of extra cook-time file region metadata in its packaging process. */
-	CookFileRegionMetadata,
-
-	/** The platform supports communication (reading and writing data) between a target a connected PC. */
-	DirectDataExchange,
-
-	/** The platform supports Luminance + Alpha encoding mode for normalmaps */
-	NormalmapLAEncodingMode,
-
-	/** All devices of this platform should be grouped under one platform group */
-	ShowAsPlatformGroup,
-
-	/** Does the platform allow various connection types to be used (ie: wifi and usb) */
-	SupportsMultipleConnectionTypes,
-
-	/** The platform can cook packages (e.g. CookedCooker) */
-	CanCookPackages
-};
-
-enum class EPlatformAuthentication
-{
-	Never,
-	Possible,
-	Always,
-};
-
-enum class EOfflineBVHMode
-{
-	Disabled,
-	MaximizePerformance,
-	MinimizeMemory,
-};
-
-/**
- * Flags specifying what is needed to be able to complete and deploy a build.
- */
-namespace ETargetPlatformReadyStatus
-{
-	/** Ready */
-	inline const int32 Ready = 0;
-
-	/** SDK Not Found*/
-	inline const int32 SDKNotFound = 1;
-
-	/** Code Build Not Supported */
-	inline const int32 CodeUnsupported = 2;
-
-	/** Plugins Not Supported */
-	inline const int32 PluginsUnsupported = 4;
-
-	/** Signing Key Not Found */
-	inline const int32 SigningKeyNotFound = 8;
-
-	/** Provision Not Found */
-	inline const int32 ProvisionNotFound = 16;
-
-	/** Manifest Not Found */
-	inline const int32 ManifestNotFound = 32;
-
-	/** Remote Server Name Empty */
-	inline const int32 RemoveServerNameEmpty = 64;
-
-	/** License Not Accepted  */
-	inline const int32 LicenseNotAccepted = 128;
-
-	/** Code Build Required */
-	inline const int32 CodeBuildRequired = 256;
-};
-
 
 /**
  * Interface for target platforms.
  *
- * This interface provides an abstraction for cooking platforms and enumerating actual target devices.
+ * This interface is only left for compatibility with platforms that have not been updated to TargetPlatformSettings and TargetPlatformControls
+ * Please only add new functions to ITargetPlatformSettings(if the sdk is not required) or ITargetPlatformControls(if the sdk is required)
  */
-class ITargetPlatform
+class ITargetPlatform : public ITargetPlatformSettings, public ITargetPlatformControls
 {
 public:
+	const ITargetPlatformSettings& GetPlatformSettings() const{ return *this; }
+	const ITargetPlatformControls& GetPlatformControls() const  { return *this; }
 
-	/**
-	 * Add a target device by name.
-	 *
-	 * @param DeviceName The name of the device to add.
-	 * @param bDefault Whether the added device should be the default.
-	 * @return true if the device was added, false otherwise.
-	 */
-	virtual bool AddDevice( const FString& DeviceName, bool bDefault ) = 0;
-
-	/**
-	 * Add a target device.
-	 *
-	 * @param DeviceId The id of the device to add.
-	 * @param DeviceUserFriendlyName The user friendly name of the device to add.
-	 * @param Username The username for the device to add.
-	 * @param Password The password for the device to add.
-	 * @param bDefault Whether the added device should be the default.
-	 * @return true if the device was added, false otherwise.
-	 */
-	virtual bool AddDevice(const FString& DeviceId, const FString& DeviceUserFriendlyName, const FString& Username, const FString& Password, bool bDefault) = 0;
-
-	/**
-	* Returns the name of this platform
-	*
-	* @return Platform name.
-	* @see DisplayName
-	*/
-	virtual FString PlatformName() const = 0;
-
-	/**
-	 * Gets the platform's display name.
-	 *
-	 * @see PlatformName
-	 */
-	virtual FText DisplayName() const = 0;
-
-	/**
-	 * Checks whether the platform's build requirements are met so that we can do things like package for the platform.
-	 *
-	 * @param bProjectHasCode true if the project has code, and therefore any compilation based SDK requirements should be checked.
-	 * @param Configuration The configuration being built
-	 * @param bRequiresAssetNativization Whether asset nativization is required
-	 * @param OutTutorialPath Let's the platform tell the editor a path to show some information about how to fix any problem.
-	 * @param OutDocumentationPath Let's the platform tell the editor a documentation path.
-	 * @param CustomizedLogMessage Let's the platform return a customized log message instead of the default for the returned status.
-	 * @return A mask of ETargetPlatformReadyStatus flags to indicate missing requirements, or 0 if all requirements are met.
-	 */
-	virtual int32 CheckRequirements(bool bProjectHasCode, EBuildConfiguration Configuration, bool bRequiresAssetNativization, FString& OutTutorialPath, FString& OutDocumentationPath, FText& CustomizedLogMessage) const = 0;
-
-	/**
-	 * Checks whether the current project needs a temporary .target.cs file to be packaged as a code project.
-	 *
-	 * @param bProjectHasCode Whether the project has code
-	 * @param Configuration The configuration being built
-	 * @param bRequiresAssetNativization Whether asset nativization is enabled
-	 * @param OutReason On success, includes a description of the reason why a target is required
-	 * @return True if a temporary target is required
-	 */
-	virtual bool RequiresTempTarget(bool bProjectHasCode, EBuildConfiguration Configuration, bool bRequiresAssetNativization, FText& OutReason) const = 0;
-
-	/**
-	 * Returns the information about this target platform
-	 */
-	virtual const PlatformInfo::FTargetPlatformInfo& GetTargetPlatformInfo() const = 0;
-
-	/**
-	 * Returns the information about the platform as a whole
-	 */
-	virtual const FDataDrivenPlatformInfo& GetPlatformInfo() const = 0;
-
-	/**
-	 * Returns the config system object usable by this TargetPlatform. It should not be modified in anyway
-	 */
-	virtual FConfigCacheIni* GetConfigSystem() const = 0;
-
-	/**
-	 * Gets the platform's INI name (so an offline tool can load the INI for the given target platform).
-	 *
-	 * @see PlatformName
-	 */
-	virtual FString IniPlatformName() const = 0;
-
-	/**
-	 * Gets the name of the device profile to use when cooking this TargetPlatform
-	 */
-	virtual FString CookingDeviceProfileName() const = 0;
-
-	/**
-	 * Enables/Disable the device check
-	 */
-	virtual void EnableDeviceCheck(bool OnOff) = 0;
-
-	/**
-	 * Returns all discoverable physical devices.
-	 *
-	 * @param OutDevices Will contain a list of discovered devices.
-	 */
-	virtual void GetAllDevices( TArray<ITargetDevicePtr>& OutDevices ) const = 0;
-
-	/**
-	 * Gets a new compression format to use in place of Zlib. This should be rarely implemented
-	 *
-	 * @return Compression format to use instead of Zlib
-	 */
-	virtual FName GetZlibReplacementFormat() const = 0;
-
-	/**
-	 * Gets the alignment of memory mapping for this platform, typically the page size.
-	 *
-	 * @return alignment of memory mapping.
-	 */
-	virtual int64 GetMemoryMappingAlignment() const = 0;
-
-	/**
-	 * Generates a platform specific asset manifest given an array of FAssetData.
-	 *
-	 * @param PakchunkMap A map of asset path to Pakchunk file indices for all of the assets.
-	 * @param PakchunkIndicesInUse A set of all Pakchunk file indices used by this set of assets.
-	 * @return true if the manifest was successfully generated, or if the platform doesn't need a manifest .
-	 */
-	virtual bool GenerateStreamingInstallManifest( const TMultiMap<FString, int32>& PakchunkMap, const TSet<int32>& PakchunkIndicesInUse) const = 0;
-
-	/**
-	 * Gets the default device.
-	 *
-	 * Note that not all platforms may have a notion of default devices.
-	 *
-	 * @return Default device.
-	 */
-	virtual ITargetDevicePtr GetDefaultDevice() const = 0;
-
-	/** 
-	 * Gets an interface to the specified device.
-	 *
-	 * @param DeviceId The identifier of the device to get.
-	 * @return The target device (can be nullptr). 
-	 */
-	virtual ITargetDevicePtr GetDevice( const FTargetDeviceId& DeviceId ) = 0;
-
-	/**
-	 * Checks whether this platform has only editor data (typically desktop platforms).
-	 *
-	 * @return true if this platform has editor only data, false otherwise.
-	 */
-	virtual bool HasEditorOnlyData() const = 0;
-
-	/**
-	 * Checks whether this platform will allow editor objects to be cooked, as opposed to editoronly properties. This will allow a
-	 * target platform to cook editoronly objects, but as if they were being cooked for a client. This is useful for a cooked editor
-	 * scenario, where every pacakge is cooked, editor and game alike.
-	 *
-	 * @return true if this platform allows editor objects to be cooked, false otherwise.
-	 */
-	virtual bool AllowsEditorObjects() const = 0;
-
-	/**
-	 * Checks whether this platform will allow development objects to be cooked. This is separate from AllowsEditorObjects
-	 * because cooked editors can be shipped with editor objects but still need to remove development assets.
-	 * 
-	 * @return true if this platform allows development objects to be cooked, false otherwise.
-	 */
-	virtual bool AllowsDevelopmentObjects() const = 0;
-
-	/**
-	 * Checks whether this platform is only a client (and must connect to a server to run).
-	 *
-	 * @return true if this platform must connect to a server.
-	 */
-	virtual bool IsClientOnly() const = 0;
-
-	/**
-	 * Checks whether this platform is little endian.
-	 *
-	 * @return true if this platform is little-endian, false otherwise.
-	 */
-	virtual bool IsLittleEndian() const = 0;
-
-	/**
-	 * Checks whether this platform is the platform that's currently running.
-	 *
-	 * For example, when running on Windows, the Windows ITargetPlatform will return true
-	 * and all other platforms will return false.
-	 *
-	 * @return true if this platform is running, false otherwise.
-	 */
-	virtual bool IsRunningPlatform() const = 0;
-
-	/**
-	 * Checks whether this platform is only a server.
-	 *
-	 * @return true if this platform has no graphics or audio, etc, false otherwise.
-	 */
-	virtual bool IsServerOnly() const = 0;
-
-	/**
-	 * Checks whether this platform is enabled for the given plugin in the currently active project.
-	 */
-	virtual bool IsEnabledForPlugin(const IPlugin& Plugin) const = 0;
-
-	/**
-	* Checks whether this platform supports shader compilation over XGE interface.
-	*
-	* @return true if this platform can distribute shader compilation threads with XGE.
-	*/
-	virtual bool CanSupportRemoteShaderCompile() const = 0;
-
-	/**
-	* Provide platform specific file dependency patterns for SN-DBS shader compilation.
-	*
-	* @param OutDependencies Platform specific dependency file patterns are uniquely appended to this array.
-	*/
-	virtual void GetShaderCompilerDependencies(TArray<FString>& OutDependencies) const = 0;
-
-	/**
-	 * Checks whether the platform's SDK requirements are met so that we can do things like
-	 * package for the platform
-	 *
-	 * @param bProjectHasCode true if the project has code, and therefore any compilation based SDK requirements should be checked
-	 * @param OutDocumentationPath Let's the platform tell the editor a path to show some documentation about how to set up the SDK
-	 * @return true if the platform is ready for use
-	 */
-	virtual bool IsSdkInstalled(bool bProjectHasCode, FString& OutDocumentationPath) const = 0;
-
-	/**
-	 * Checks whether this platform requires cooked data (typically console platforms).
-	 *
-	 * @return true if this platform requires cooked data, false otherwise.
-	 */
-	virtual bool RequiresCookedData() const = 0;
-
-	/**
-	 * Checks whether this platform requires the originally released version (in addition to the 
-	 * previously released version) to create a patch
-	 *
-	 * @return true if this platform requires the originally released version, false otherwise.
-	 */
-	virtual bool RequiresOriginalReleaseVersionForPatch() const = 0;
-
-	/**
-	* Checks whether this platform has a secure shippable package format, and therefore doesn't need any encryption or signing support
-	*
-	* @return true if this platform requires cooked data, false otherwise.
-	*/
-	virtual bool HasSecurePackageFormat() const = 0;
-	
-	/**
-	 * Checks whether this platform requires user credentials (typically server platforms).
-	 *
-	 * @return enum if this platform requires user credentials.
-	 */
-	virtual EPlatformAuthentication RequiresUserCredentials() const = 0;
-
-	/**
-	 * Returns true if the platform supports the AutoSDK system
-	 */
-	virtual bool SupportsAutoSDK() const = 0;
-
-	/**
-	 * Checks whether this platform supports the specified build target, i.e. Game or Editor.
-	 *
-	 * @param TargetType The target type to check.
-	 * @return true if the build target is supported, false otherwise.
-	 */
-	virtual bool SupportsBuildTarget( EBuildTargetType TargetType ) const = 0;
-
-	/**
-	 * Return the TargetType this platform uses at runtime.
-	 * Some TargetPlatforms like CookedEditors need to cook with one type, like Client, but then will run as an Editor.
-	 * Decisions made based solely on the cook type may cause data mismatches if the runtime type is different.
-	 * This is also useful for knowing what plugins will be enabled at runtime.
-	 */
-	virtual EBuildTargetType GetRuntimePlatformType() const = 0;
-
-	/**
-	 * Checks whether the target platform supports the specified feature.
-	 *
-	 * @param Feature The feature to check.
-	 * @return true if the feature is supported, false otherwise.
-	 */
-	virtual bool SupportsFeature( ETargetPlatformFeatures Feature ) const = 0;
-
-	/**
-	 * Checks whether the target platform supports the specified value for the specified type of support
-	 *
-	 * @param SupportedType The type of support being queried
-	 * @param RequiredSupportedValue The value of support needed
-	 * @return true if the feature is supported, false otherwise.
-	 */
-	virtual bool SupportsValueForType(FName SupportedType, FName RequiredSupportedValue) const = 0;
-
-	/**
-	 * Gets whether the platform should use forward shading or not.
-	 */
-	virtual bool UsesForwardShading() const = 0;
-
-	/**
-	* Gets whether the platform should use DBuffer for decals.
-	*/
-	virtual bool UsesDBuffer() const = 0;
-	
-	/**
-	* Gets whether the platform should output velocity in the base pass.
-	*/
-	virtual bool UsesBasePassVelocity() const = 0;
-
-	/**
-	* Gets whether the platform will use selective outputs in the base pass shaders.
-	*/
-	virtual bool UsesSelectiveBasePassOutputs() const = 0; 
-
-	/**
-	* Gets whether the platform will use distance fields.
-	*/
-	virtual bool UsesDistanceFields() const = 0;
-
-	/**
-	* Gets whether the platform will use ray tracing.
-	*/
-	virtual bool UsesRayTracing() const = 0;
-
-	/**
-	* Gets static mesh offline BVH mode
-	*/
-	virtual EOfflineBVHMode GetStaticMeshOfflineBVHMode() const = 0;
-
-	/**
-	* Gets whether the platform will use compression for static mesh offline BVH.
-	*/
-	virtual bool GetStaticMeshOfflineBVHCompression() const = 0;
-
-	/**
-	* Gets whether the platform will use SH2 instead of SH3 for sky irradiance.
-	*/
-	virtual bool ForcesSimpleSkyDiffuse() const = 0;
-
-	/**
-	* Gets whether the platform will encode depth velocity.
-	*/
-	virtual bool VelocityEncodeDepth() const = 0;
-
-	/**
-	* Gets down sample mesh distance field divider.
-	*
-	* @return 1 if platform does not need to downsample mesh distance fields
-	*/
-	virtual float GetDownSampleMeshDistanceFieldDivider() const = 0;
-
-	/**
-	* Gets an integer representing the height fog mode for opaque materials on a platform.
-	* @return 0 if no override (i.e. use r.VertexFoggingForOpaque from project settings); 1 if pixel fog; 2 if vertex fog.
-	*/
-	virtual int32 GetHeightFogModeForOpaque() const = 0;
-
-	/**
-	 * Gets whether the platform uses Mobile AO
-	 */
-	virtual bool UsesMobileAmbientOcclusion() const = 0;
-
-	/**
-	 * Gets whether the platform should use DBuffer for decals when using the mobile renderer.
-	 */
-	virtual bool UsesMobileDBuffer() const = 0;
-	
-	/**
-	 * Gets whether the platform uses ASTC HDR
-	 */
-	virtual bool UsesASTCHDR() const = 0;
-
-	/**
-	* Gets the shader formats this platform can use.
-	*
-	* @param OutFormats Will contain the shader formats.
-	*/
-	virtual void GetAllPossibleShaderFormats( TArray<FName>& OutFormats ) const = 0;
-
-	/**
-	* Gets the shader formats that have been selected for this target platform
-	*
-	* @param OutFormats Will contain the shader formats.
-	*/
-	virtual void GetAllTargetedShaderFormats(TArray<FName>& OutFormats) const = 0;
-
-	/**
-	* Gets the shader formats that support ray tracing for this target platform.
-	* 
-	* @param OutFormats Will contain the shader formats.
-	*/
-	virtual void GetRayTracingShaderFormats(TArray<FName>& OutFormats) const = 0;
-
-	/**
-	 * Gather per-project cook/package analytics
-	 */
-	virtual void GetPlatformSpecificProjectAnalytics( TArray<struct FAnalyticsEventAttribute>& AnalyticsParamArray ) const = 0;
-
-#if WITH_ENGINE
-	/**
-	 * Gets the format to use for a particular body setup.
-	 *
-	 * @return Physics format.
-	 */
-	virtual FName GetPhysicsFormat( class UBodySetup* Body ) const = 0;
-
-	/**
-	 * Gets the reflection capture formats this platform needs.
-	 *
-	 * @param OutFormats Will contain the collection of formats.
-	 */
-	virtual void GetReflectionCaptureFormats( TArray<FName>& OutFormats ) const = 0;
-
-	/**
-	 * Gets a list of modules that may contain the GetAllTargetedShaderFormats. This is optional -
-	 * if any required shader format isn't found in this list, then it will use the old path
-	 * of loading all shader format modules to gather all available shader formats
-	 */
-	virtual void GetShaderFormatModuleHints(TArray<FName>& OutModuleNames) const = 0;
-
-	/**
-	 * Gets the texture format to use for each layer in the given texture, for each of the platform's formats.
-	 * _Most_ platforms only supply one format for a given texture, so OutFormats.Num() is usually 1. The exception is Android_Multi,
-	 * where you can get several formats due to targeting different devices.
-	 * OutFormats.Num() == NumberOfPlatformFormats
-	 * OutFormats[0...N].Num() == Texture->Source.GetNumLayers().
-	 */
-	virtual void GetTextureFormats( const class UTexture* Texture, TArray< TArray<FName> >& OutFormats ) const = 0;
-
-	/**
-	 * Gets the texture formats this platform can use
-	 * 
-	 * @param OutFormats will contain all the texture formats which are possible for this platform
-	 */
-	virtual void GetAllTextureFormats( TArray<FName>& OutFormats ) const = 0;
-
-	/**
-	 * Gets a list of modules that may contain the GetAllTextureFormats. This is optional -
-	 * if any required texture format isn't found in this list, then it will use the old path
-	 * of loading all texture format modules to gather all available texture formats
-	 */
-	virtual void GetTextureFormatModuleHints(TArray<FName>& OutModuleNames) const = 0;
-
-	/**
-	 * Platforms that support multiple texture compression variants 
-	 * might want to use a single specific variant for virtual textures to reduce fragmentation
-	 */
-	virtual FName FinalizeVirtualTextureLayerFormat(FName Format) const = 0;
-		
-	//whether R5G6B5 and B5G5R5A1 is supported
-	virtual bool SupportsLQCompressionTextureFormat() const = 0;
-	
-	/**
-	 * Gets the format to use for a particular piece of audio.
-	 *
-	 * @param Wave The sound node wave to get the format for.
-	 * @return Name of the wave format.
-	 */
-	virtual FName GetWaveFormat( const class USoundWave* Wave ) const = 0;
-
-	/**
-	* Gets all the formats which can be returned from GetWaveFormat
-	*
-	* @param output array of all the formats
-	*/
-	virtual void GetAllWaveFormats( TArray<FName>& OutFormats ) const = 0;
-
-	/**
-	 * Gets a list of modules that may contain the GetAllWaveFormats. This is optional -
-	 * if any required audio format isn't found in this list, then it will use the old path
-	 * of loading all audio format modules to gather all available audio formats
-	 */
-	virtual void GetWaveFormatModuleHints(TArray<FName>& OutModuleNames) const = 0;
-
-	/**
-	 * Checks whether if this platform wants AV data (defaults to !IsServerOnly(), which is the standard reason why we don't want AV data)
-	 * Used so that custom target platforms can remove AV data, but is not a server-only platform
-	 *
-	 * @return true if this platform allows AV data to be cooked, false otherwise.
-	 */
-	virtual bool AllowAudioVisualData() const = 0;
-
-	/** Checks if this Target will want to load this object (generally used to mark an object to not be cooked for this target) */
-	virtual bool AllowObject(const class UObject* Object) const = 0;
-
-	/**
-	 * Gets the texture LOD settings used by this platform.
-	 *
-	 * @return A texture LOD settings structure.
-	 */
-	virtual const class UTextureLODSettings& GetTextureLODSettings() const = 0;
-
-	/**
-	* Register Basic LOD Settings for this platform
-	*/
-	virtual void RegisterTextureLODSettings(const class UTextureLODSettings* InTextureLODSettings) = 0;
-
-	/**
-	 * Gets the static mesh LOD settings used by this platform.
-	 *
-	 * @return A static mesh LOD settings structure.
-	 */
-	virtual const class FStaticMeshLODSettings& GetStaticMeshLODSettings() const = 0;
-
-	/** 
-	 * Gets the name of the mesh builder module that is responsible for 
-	 * building static and skeletal mesh derived data for this platform.
-	 *
-	 * @return The name of a mesh builder module.
-	 */
-	virtual FName GetMeshBuilderModuleName() const = 0;
-#endif
-
-	/** 
-	 * Package a build for the given platform 
-	 * 
-	 * @param InPackgeDirectory The directory that contains what needs to be packaged.
-	 * @return bool true on success, false otherwise.
-	 */ 
-	virtual bool PackageBuild(const FString& InPackgeDirectory) = 0;
-
-	/**
-	 * Returns true if the platform is part of a family of variants
-	 */
-	virtual bool SupportsVariants() const = 0;
-
-	/**
-	 * Gets the variant priority of this platform
-	 *
-	 * @return float priority for this platform variant.
-	 */
-	virtual float GetVariantPriority() const = 0;
-
-	/** 
-	 * Whether or not to send all lower-case filepaths when connecting over a fileserver connection. 
-	 */
-	virtual bool SendLowerCaseFilePaths() const = 0;
-
-	/**
-	 * Project settings to check to determine if a build should occur
-	 */
-	virtual void GetBuildProjectSettingKeys(FString& OutSection, TArray<FString>& InBoolKeys, TArray<FString>& InIntKeys, TArray<FString>& InStringKeys) const = 0;
-
-	/**
-	 * Get unique integer identifier for this platform.
-	 *
-	 * The implementation will assign an ordinal to each target platform at startup, assigning
-	 * a value of 0, 1, 2, etc in order to make the ordinals usable as array / bit mask indices.
-	 *
-	 * @return int32 A unique integer which may be used to identify target platform during the 
-	 *               current session only (note: not stable across runs).
-	 */
-	virtual int32 GetPlatformOrdinal() const = 0;
 
+public:
 	/**
 	 * Given a platform ordinal number, returns the corresponding ITargetPlatform instance
 	 */
 	TARGETPLATFORM_API static const ITargetPlatform* GetPlatformFromOrdinal(int32 Ordinal);
-
-	/**
-	 * Returns custom DeviceManager widget creator for this platform
-	 */
-	virtual TSharedPtr<IDeviceManagerCustomPlatformWidgetCreator> GetCustomWidgetCreator() const = 0;
-
-	/**
-	 * Returns wheter or not this 16bit index buffer should be promoted to 32bit
-	 */
-	virtual bool ShouldExpandTo32Bit(const uint16* Indices, const int32 NumIndices) const = 0;
-
-	/**
-	 * Copy a file to the target
-	 */
-	virtual bool CopyFileToTarget(const FString& DeviceId, const FString& HostFilename, const FString& TargetFilename, const TMap<FString,FString>& CustomPlatformData) = 0;
-
-	/**
-	 * Gets a list of package names to cook when cooking this platform
-	 */
-	virtual void GetExtraPackagesToCook(TArray<FName>& PackageNames) const = 0;
-
-	/**
-	 * Initializes the host platform to support target devices (may be called multiple times after an SDK is installed while running)
-	 */
-	virtual bool InitializeHostPlatform() = 0;
-
-public:
-
 	/**
 	 * Gets an event delegate that is executed when a new target device has been discovered.
 	 */
diff --git a/Engine/Source/Developer/TargetPlatform/Public/Interfaces/ITargetPlatformControls.h b/Engine/Source/Developer/TargetPlatform/Public/Interfaces/ITargetPlatformControls.h
new file mode 100644
--- /dev/null
+++ b/Engine/Source/Developer/TargetPlatform/Public/Interfaces/ITargetPlatformControls.h
@@ -0,0 +1,515 @@
+// Copyright Epic Games, Inc. All Rights Reserved.
+
+#pragma once
+#include "Interfaces/ITargetDevice.h"
+
+class IDeviceManagerCustomPlatformWidgetCreator;
+class IPlugin;
+struct FDataDrivenPlatformInfo;
+class ITargetPlatformSettings;
+namespace PlatformInfo
+{
+	// Forward declare type from DesktopPlatform rather than add an include dependency to everything using ITargetPlatform
+	struct FTargetPlatformInfo;
+}
+
+
+enum class EPlatformAuthentication
+{
+	Never,
+	Possible,
+	Always,
+};
+
+/**
+ * Flags specifying what is needed to be able to complete and deploy a build.
+ */
+namespace ETargetPlatformReadyStatus
+{
+	/** Ready */
+	inline const int32 Ready = 0;
+
+	/** SDK Not Found*/
+	inline const int32 SDKNotFound = 1;
+
+	/** Code Build Not Supported */
+	inline const int32 CodeUnsupported = 2;
+
+	/** Plugins Not Supported */
+	inline const int32 PluginsUnsupported = 4;
+
+	/** Signing Key Not Found */
+	inline const int32 SigningKeyNotFound = 8;
+
+	/** Provision Not Found */
+	inline const int32 ProvisionNotFound = 16;
+
+	/** Manifest Not Found */
+	inline const int32 ManifestNotFound = 32;
+
+	/** Remote Server Name Empty */
+	inline const int32 RemoveServerNameEmpty = 64;
+
+	/** License Not Accepted  */
+	inline const int32 LicenseNotAccepted = 128;
+
+	/** Code Build Required */
+	inline const int32 CodeBuildRequired = 256;
+};
+
+
+class ITargetPlatformControls
+{
+public:
+	ITargetPlatformControls() {};
+	ITargetPlatformControls(ITargetPlatformSettings* InTargetPlatformSettings) :TargetPlatformSettings(InTargetPlatformSettings) {};
+	/**
+	 * Add a target device by name.
+	 *
+	 * @param DeviceName The name of the device to add.
+	 * @param bDefault Whether the added device should be the default.
+	 * @return true if the device was added, false otherwise.
+	 */
+	virtual bool AddDevice(const FString& DeviceName, bool bDefault) = 0;
+
+	/**
+	 * Add a target device.
+	 *
+	 * @param DeviceId The id of the device to add.
+	 * @param DeviceUserFriendlyName The user friendly name of the device to add.
+	 * @param Username The username for the device to add.
+	 * @param Password The password for the device to add.
+	 * @param bDefault Whether the added device should be the default.
+	 * @return true if the device was added, false otherwise.
+	 */
+	virtual bool AddDevice(const FString& DeviceId, const FString& DeviceUserFriendlyName, const FString& Username, const FString& Password, bool bDefault) = 0;
+
+	/**
+	* Returns the name of this platform
+	*
+	* @return Platform name.
+	* @see DisplayName
+	*/
+	virtual FString PlatformName() const = 0;
+
+	/**
+	 * Gets the platform's display name.
+	 *
+	 * @see PlatformName
+	 */
+	virtual FText DisplayName() const = 0;
+
+	/**
+	 * Checks whether the platform's build requirements are met so that we can do things like package for the platform.
+	 *
+	 * @param bProjectHasCode true if the project has code, and therefore any compilation based SDK requirements should be checked.
+	 * @param Configuration The configuration being built
+	 * @param bRequiresAssetNativization Whether asset nativization is required
+	 * @param OutTutorialPath Let's the platform tell the editor a path to show some information about how to fix any problem.
+	 * @param OutDocumentationPath Let's the platform tell the editor a documentation path.
+	 * @param CustomizedLogMessage Let's the platform return a customized log message instead of the default for the returned status.
+	 * @return A mask of ETargetPlatformReadyStatus flags to indicate missing requirements, or 0 if all requirements are met.
+	 */
+	virtual int32 CheckRequirements(bool bProjectHasCode, EBuildConfiguration Configuration, bool bRequiresAssetNativization, FString& OutTutorialPath, FString& OutDocumentationPath, FText& CustomizedLogMessage) const = 0;
+
+	/**
+	 * Checks whether the current project needs a temporary .target.cs file to be packaged as a code project.
+	 *
+	 * @param bProjectHasCode Whether the project has code
+	 * @param Configuration The configuration being built
+	 * @param bRequiresAssetNativization Whether asset nativization is enabled
+	 * @param OutReason On success, includes a description of the reason why a target is required
+	 * @return True if a temporary target is required
+	 */
+	virtual bool RequiresTempTarget(bool bProjectHasCode, EBuildConfiguration Configuration, bool bRequiresAssetNativization, FText& OutReason) const = 0;
+
+	/**
+	 * Returns the information about this target platform
+	 */
+	virtual const PlatformInfo::FTargetPlatformInfo& GetTargetPlatformInfo() const = 0;
+
+	/**
+	 * Returns the information about the platform as a whole
+	 */
+	virtual const FDataDrivenPlatformInfo& GetPlatformInfo() const = 0;
+
+	/**
+	 * Gets the name of the device profile to use when cooking this TargetPlatform
+	 */
+	virtual FString CookingDeviceProfileName() const = 0;
+
+	/**
+	 * Enables/Disable the device check
+	 */
+	virtual void EnableDeviceCheck(bool OnOff) = 0;
+
+	/**
+	 * Returns all discoverable physical devices.
+	 *
+	 * @param OutDevices Will contain a list of discovered devices.
+	 */
+	virtual void GetAllDevices(TArray<ITargetDevicePtr>& OutDevices) const = 0;
+	
+	/**
+	 * Gets a new compression format to use in place of Zlib. This should be rarely implemented
+	 *
+	 * @return Compression format to use instead of Zlib
+	 */
+	virtual FName GetZlibReplacementFormat() const = 0;
+
+	/**
+	 * Gets the alignment of memory mapping for this platform, typically the page size.
+	 *
+	 * @return alignment of memory mapping.
+	 */
+	virtual int64 GetMemoryMappingAlignment() const = 0;
+	
+	/**
+	 * Generates a platform specific asset manifest given an array of FAssetData.
+	 *
+	 * @param PakchunkMap A map of asset path to Pakchunk file indices for all of the assets.
+	 * @param PakchunkIndicesInUse A set of all Pakchunk file indices used by this set of assets.
+	 * @return true if the manifest was successfully generated, or if the platform doesn't need a manifest .
+	 */
+	virtual bool GenerateStreamingInstallManifest(const TMultiMap<FString, int32>& PakchunkMap, const TSet<int32>& PakchunkIndicesInUse) const = 0;
+
+	/**
+	 * Gets the default device.
+	 *
+	 * Note that not all platforms may have a notion of default devices.
+	 *
+	 * @return Default device.
+	 */
+	virtual ITargetDevicePtr GetDefaultDevice() const = 0;
+
+	/**
+	 * Gets an interface to the specified device.
+	 *
+	 * @param DeviceId The identifier of the device to get.
+	 * @return The target device (can be nullptr).
+	 */
+	virtual ITargetDevicePtr GetDevice(const FTargetDeviceId& DeviceId) = 0;
+
+	/**
+	 * Checks whether this platform has only editor data (typically desktop platforms).
+	 *
+	 * @return true if this platform has editor only data, false otherwise.
+	 */
+	virtual bool HasEditorOnlyData() const = 0;
+
+	/**
+	 * Checks whether this platform will allow editor objects to be cooked, as opposed to editoronly properties. This will allow a
+	 * target platform to cook editoronly objects, but as if they were being cooked for a client. This is useful for a cooked editor
+	 * scenario, where every pacakge is cooked, editor and game alike.
+	 *
+	 * @return true if this platform allows editor objects to be cooked, false otherwise.
+	 */
+	virtual bool AllowsEditorObjects() const = 0;
+
+	/**
+	 * Checks whether this platform will allow development objects to be cooked. This is separate from AllowsEditorObjects
+	 * because cooked editors can be shipped with editor objects but still need to remove development assets.
+	 *
+	 * @return true if this platform allows development objects to be cooked, false otherwise.
+	 */
+	virtual bool AllowsDevelopmentObjects() const = 0;
+
+	/**
+	 * Checks whether this platform is only a client (and must connect to a server to run).
+	 *
+	 * @return true if this platform must connect to a server.
+	 */
+	virtual bool IsClientOnly() const = 0;
+
+	/**
+	 * Checks whether this platform is little endian.
+	 *
+	 * @return true if this platform is little-endian, false otherwise.
+	 */
+	virtual bool IsLittleEndian() const = 0;
+
+	/**
+	 * Checks whether this platform is the platform that's currently running.
+	 *
+	 * For example, when running on Windows, the Windows ITargetPlatform will return true
+	 * and all other platforms will return false.
+	 *
+	 * @return true if this platform is running, false otherwise.
+	 */
+	virtual bool IsRunningPlatform() const = 0;
+
+	/**
+	 * Checks whether this platform is only a server.
+	 *
+	 * @return true if this platform has no graphics or audio, etc, false otherwise.
+	 */
+	virtual bool IsServerOnly() const = 0;
+
+	/**
+	 * Checks whether this platform is enabled for the given plugin in the currently active project.
+	 */
+	virtual bool IsEnabledForPlugin(const IPlugin& Plugin) const = 0;
+
+	/**
+	* Checks whether this platform supports shader compilation over XGE interface.
+	*
+	* @return true if this platform can distribute shader compilation threads with XGE.
+	*/
+	virtual bool CanSupportRemoteShaderCompile() const = 0;
+
+	/**
+	* Provide platform specific file dependency patterns for SN-DBS shader compilation.
+	*
+	* @param OutDependencies Platform specific dependency file patterns are uniquely appended to this array.
+	*/
+	virtual void GetShaderCompilerDependencies(TArray<FString>& OutDependencies) const = 0;
+
+	/**
+	 * Checks whether the platform's SDK requirements are met so that we can do things like
+	 * package for the platform
+	 *
+	 * @param bProjectHasCode true if the project has code, and therefore any compilation based SDK requirements should be checked
+	 * @param OutDocumentationPath Let's the platform tell the editor a path to show some documentation about how to set up the SDK
+	 * @return true if the platform is ready for use
+	 */
+	virtual bool IsSdkInstalled(bool bProjectHasCode, FString& OutDocumentationPath) const = 0;
+
+	/**
+	 * Checks whether this platform requires cooked data (typically console platforms).
+	 *
+	 * @return true if this platform requires cooked data, false otherwise.
+	 */
+	virtual bool RequiresCookedData() const = 0;
+
+	/**
+	 * Checks whether this platform requires the originally released version (in addition to the
+	 * previously released version) to create a patch
+	 *
+	 * @return true if this platform requires the originally released version, false otherwise.
+	 */
+	virtual bool RequiresOriginalReleaseVersionForPatch() const = 0;
+
+	/**
+	* Checks whether this platform has a secure shippable package format, and therefore doesn't need any encryption or signing support
+	*
+	* @return true if this platform requires cooked data, false otherwise.
+	*/
+	virtual bool HasSecurePackageFormat() const = 0;
+	
+	/**
+	 * Checks whether this platform requires user credentials (typically server platforms).
+	 *
+	 * @return enum if this platform requires user credentials.
+	 */
+	virtual EPlatformAuthentication RequiresUserCredentials() const = 0;
+
+	/**
+	 * Returns true if the platform supports the AutoSDK system
+	 */
+	virtual bool SupportsAutoSDK() const = 0;
+
+	/**
+	 * Checks whether this platform supports the specified build target, i.e. Game or Editor.
+	 *
+	 * @param TargetType The target type to check.
+	 * @return true if the build target is supported, false otherwise.
+	 */
+	virtual bool SupportsBuildTarget(EBuildTargetType TargetType) const = 0;
+
+	/**
+	 * Return the TargetType this platform uses at runtime.
+	 * Some TargetPlatforms like CookedEditors need to cook with one type, like Client, but then will run as an Editor.
+	 * Decisions made based solely on the cook type may cause data mismatches if the runtime type is different.
+	 * This is also useful for knowing what plugins will be enabled at runtime.
+	 */
+	virtual EBuildTargetType GetRuntimePlatformType() const = 0;
+	
+	/**
+	 * Gather per-project cook/package analytics
+	 */
+	virtual void GetPlatformSpecificProjectAnalytics(TArray<struct FAnalyticsEventAttribute>& AnalyticsParamArray) const = 0;
+
+#if WITH_ENGINE
+	/**
+	 * Gets the format to use for a particular body setup.
+	 *
+	 * @return Physics format.
+	 */
+	virtual FName GetPhysicsFormat(class UBodySetup* Body) const = 0;
+
+	/**
+	 * Gets a list of modules that may contain the GetAllTargetedShaderFormats. This is optional -
+	 * if any required shader format isn't found in this list, then it will use the old path
+	 * of loading all shader format modules to gather all available shader formats
+	 */
+	virtual void GetShaderFormatModuleHints(TArray<FName>& OutModuleNames) const = 0;
+
+	/**
+	 * Gets the texture format to use for each layer in the given texture, for each of the platform's formats.
+	 * _Most_ platforms only supply one format for a given texture, so OutFormats.Num() is usually 1. The exception is Android_Multi,
+	 * where you can get several formats due to targeting different devices.
+	 * OutFormats.Num() == NumberOfPlatformFormats
+	 * OutFormats[0...N].Num() == Texture->Source.GetNumLayers().
+	 */
+	virtual void GetTextureFormats(const class UTexture* Texture, TArray< TArray<FName> >& OutFormats) const = 0;
+
+	/**
+	 * Gets the texture formats this platform can use
+	 *
+	 * @param OutFormats will contain all the texture formats which are possible for this platform
+	 */
+	virtual void GetAllTextureFormats(TArray<FName>& OutFormats) const = 0;
+
+	/**
+	 * Gets a list of modules that may contain the GetAllTextureFormats. This is optional -
+	 * if any required texture format isn't found in this list, then it will use the old path
+	 * of loading all texture format modules to gather all available texture formats
+	 */
+	virtual void GetTextureFormatModuleHints(TArray<FName>& OutModuleNames) const = 0;
+
+	/**
+	 * Platforms that support multiple texture compression variants
+	 * might want to use a single specific variant for virtual textures to reduce fragmentation
+	 */
+	virtual FName FinalizeVirtualTextureLayerFormat(FName Format) const = 0;
+
+	//whether R5G6B5 and B5G5R5A1 is supported
+	virtual bool SupportsLQCompressionTextureFormat() const = 0;
+
+	/**
+	 * Gets the format to use for a particular piece of audio.
+	 *
+	 * @param Wave The sound node wave to get the format for.
+	 * @return Name of the wave format.
+	 */
+	virtual FName GetWaveFormat(const class USoundWave* Wave) const = 0;
+
+	/**
+	* Gets all the formats which can be returned from GetWaveFormat
+	*
+	* @param output array of all the formats
+	*/
+	virtual void GetAllWaveFormats(TArray<FName>& OutFormats) const = 0;
+	
+	/**
+	 * Gets a list of modules that may contain the GetAllWaveFormats. This is optional -
+	 * if any required audio format isn't found in this list, then it will use the old path
+	 * of loading all audio format modules to gather all available audio formats
+	 */
+	virtual void GetWaveFormatModuleHints(TArray<FName>& OutModuleNames) const = 0;
+
+	/**
+	 * Checks whether if this platform wants AV data (defaults to !IsServerOnly(), which is the standard reason why we don't want AV data)
+	 * Used so that custom target platforms can remove AV data, but is not a server-only platform
+	 *
+	 * @return true if this platform allows AV data to be cooked, false otherwise.
+	 */
+	virtual bool AllowAudioVisualData() const = 0;
+
+	/** Checks if this Target will want to load this object (generally used to mark an object to not be cooked for this target) */
+	virtual bool AllowObject(const class UObject* Object) const = 0;
+
+	/**
+	 * Gets the name of the mesh builder module that is responsible for
+	 * building static and skeletal mesh derived data for this platform.
+	 *
+	 * @return The name of a mesh builder module.
+	 */
+	virtual FName GetMeshBuilderModuleName() const = 0;
+#endif
+
+	/**
+	 * Package a build for the given platform
+	 *
+	 * @param InPackgeDirectory The directory that contains what needs to be packaged.
+	 * @return bool true on success, false otherwise.
+	 */
+	virtual bool PackageBuild(const FString& InPackgeDirectory) = 0;
+
+	/**
+	 * Returns true if the platform is part of a family of variants
+	 */
+	virtual bool SupportsVariants() const = 0;
+
+	/**
+	 * Gets the variant priority of this platform
+	 *
+	 * @return float priority for this platform variant.
+	 */
+	virtual float GetVariantPriority() const = 0;
+
+	/**
+	 * Whether or not to send all lower-case filepaths when connecting over a fileserver connection.
+	 */
+	virtual bool SendLowerCaseFilePaths() const = 0;
+
+	/**
+	 * Project settings to check to determine if a build should occur
+	 */
+	virtual void GetBuildProjectSettingKeys(FString& OutSection, TArray<FString>& InBoolKeys, TArray<FString>& InIntKeys, TArray<FString>& InStringKeys) const = 0;
+
+	/**
+	 * Get unique integer identifier for this platform.
+	 *
+	 * The implementation will assign an ordinal to each target platform at startup, assigning
+	 * a value of 0, 1, 2, etc in order to make the ordinals usable as array / bit mask indices.
+	 *
+	 * @return int32 A unique integer which may be used to identify target platform during the
+	 *               current session only (note: not stable across runs).
+	 */
+	virtual int32 GetPlatformOrdinal() const = 0;
+
+	/**
+	 * Given a platform ordinal number, returns the corresponding ITargetPlatform instance
+	 */
+	TARGETPLATFORM_API static const ITargetPlatformControls* GetPlatformFromOrdinal(int32 Ordinal);
+
+	/**
+	 * Returns custom DeviceManager widget creator for this platform
+	 */
+	virtual TSharedPtr<IDeviceManagerCustomPlatformWidgetCreator> GetCustomWidgetCreator() const = 0;
+
+	/**
+	 * Returns wheter or not this 16bit index buffer should be promoted to 32bit
+	 */
+	virtual bool ShouldExpandTo32Bit(const uint16* Indices, const int32 NumIndices) const = 0;
+
+	/**
+	 * Copy a file to the target
+	 */
+	virtual bool CopyFileToTarget(const FString& DeviceId, const FString& HostFilename, const FString& TargetFilename, const TMap<FString, FString>& CustomPlatformData) = 0;
+
+	/**
+	 * Gets a list of package names to cook when cooking this platform
+	 */
+	virtual void GetExtraPackagesToCook(TArray<FName>& PackageNames) const = 0;
+
+	/**
+	 * Initializes the host platform to support target devices (may be called multiple times after an SDK is installed while running)
+	 */
+	virtual bool InitializeHostPlatform() = 0;
+
+	ITargetPlatformSettings* GetTargetPlatformSettings() const
+	{
+		return TargetPlatformSettings;
+	};
+
+	/** Virtual destructor. */
+	virtual ~ITargetPlatformControls() { }
+
+	/**
+	 * Gets an event delegate that is executed when a new target device has been discovered.
+	 */
+	DECLARE_EVENT_OneParam(ITargetPlatformControls, FOnTargetDeviceDiscovered, ITargetDeviceRef /*DiscoveredDevice*/);
+	static TARGETPLATFORM_API FOnTargetDeviceDiscovered& OnDeviceDiscovered();
+
+	/**
+	 * Gets an event delegate that is executed when a target device has been lost, i.e. disconnected or timed out.
+	 */
+	DECLARE_EVENT_OneParam(ITargetPlatformControls, FOnTargetDeviceLost, ITargetDeviceRef /*LostDevice*/);
+	static TARGETPLATFORM_API FOnTargetDeviceLost& OnDeviceLost();
+
+protected:
+	ITargetPlatformSettings* TargetPlatformSettings;
+	static TARGETPLATFORM_API int32 AssignPlatformOrdinal(const ITargetPlatformControls& Platform);
+};
\ No newline at end of file
diff --git a/Engine/Source/Developer/TargetPlatform/Public/Interfaces/ITargetPlatformControlsModule.h b/Engine/Source/Developer/TargetPlatform/Public/Interfaces/ITargetPlatformControlsModule.h
new file mode 100644
--- /dev/null
+++ b/Engine/Source/Developer/TargetPlatform/Public/Interfaces/ITargetPlatformControlsModule.h
@@ -0,0 +1,53 @@
+// Copyright Epic Games, Inc. All Rights Reserved.
+
+#pragma once
+
+#include "CoreMinimal.h"
+#include "Modules/ModuleInterface.h"
+#include "Interfaces/ITargetPlatformControls.h"
+
+/**
+ * Interface for target platform modules.
+ */
+class ITargetPlatformControlsModule
+	: public IModuleInterface
+{
+
+public:
+	/** Virtual destructor. */
+	virtual ~ITargetPlatformControlsModule()
+	{
+		for (ITargetPlatformControls* TP : AllTargetPlatforms)
+		{
+			delete TP;
+		}
+		AllTargetPlatforms.Empty();
+	}
+
+	/**
+	 * Gets the module's target platforms. This should be overridden by each platform, but
+	 * currently, we are re-using the single internal GetTargetPlatform method the old TPModules will implement
+	 *
+	 * @return The target platform.
+	 */
+	TArray<ITargetPlatformControls*> GetTargetPlatformControls(FName& PlatformSettingsModuleName)
+	{
+		if (AllTargetPlatforms.Num() == 0)
+		{
+			GetTargetPlatformControls(AllTargetPlatforms, PlatformSettingsModuleName);
+		}
+
+		return AllTargetPlatforms;
+	}
+
+protected:
+
+	/**
+	 * This is where each platform module will fill out an array
+	*/
+	virtual void GetTargetPlatformControls(TArray<ITargetPlatformControls*>& TargetPlatforms, FName& PlatformSettingsModuleName) = 0;
+
+private:
+	/** Holds the target platforms. */
+	TArray<ITargetPlatformControls*> AllTargetPlatforms;
+};
diff --git a/Engine/Source/Developer/TargetPlatform/Public/Interfaces/ITargetPlatformManagerModule.h b/Engine/Source/Developer/TargetPlatform/Public/Interfaces/ITargetPlatformManagerModule.h
--- a/Engine/Source/Developer/TargetPlatform/Public/Interfaces/ITargetPlatformManagerModule.h
+++ b/Engine/Source/Developer/TargetPlatform/Public/Interfaces/ITargetPlatformManagerModule.h
@@ -9,6 +9,8 @@
 #include "Modules/ModuleInterface.h"
 
 class ITargetPlatform;
+class ITargetPlatformControls;
+class ITargetPlatformSettings;
 
 DECLARE_MULTICAST_DELEGATE(FOnTargetPlatformsInvalidated);
 
@@ -144,6 +146,8 @@ public:
 	 * @return Collection of platforms.
 	 */
 	virtual const TArray<ITargetPlatform*>& GetTargetPlatforms() = 0;
+	virtual const TArray<ITargetPlatformControls*>& GetTargetPlatformControls() = 0;
+	virtual const TArray<ITargetPlatformSettings*>& GetTargetPlatformSettings() = 0;
 
 	/**
 	 * Returns the list of all ITextureFormats that were located in DLLs.
diff --git a/Engine/Source/Developer/TargetPlatform/Public/Interfaces/ITargetPlatformModule.h b/Engine/Source/Developer/TargetPlatform/Public/Interfaces/ITargetPlatformModule.h
--- a/Engine/Source/Developer/TargetPlatform/Public/Interfaces/ITargetPlatformModule.h
+++ b/Engine/Source/Developer/TargetPlatform/Public/Interfaces/ITargetPlatformModule.h
@@ -5,7 +5,12 @@
 #include "CoreMinimal.h"
 #include "Modules/ModuleInterface.h"
 #include "Interfaces/ITargetPlatform.h"
+#include "Interfaces/ITargetPlatformSettings.h"
 
+
+class ITargetPlatformControlsModule;
+class ITargetPlatformSettingsModule;
+class ITargetPlatformControls;
 /**
  * Interface for target platform modules.
  */
@@ -35,19 +40,25 @@ public:
 	{
 		if (AllTargetPlatforms.Num() == 0)
 		{
-			GetTargetPlatforms(AllTargetPlatforms);
+			GetTargetPlatforms(AllTargetPlatforms, PlatformSettings, PlatformControls);
 		}
 
 		return AllTargetPlatforms;
 	}
 
+	TArray<ITargetPlatformSettings*> PlatformSettings;
+	TArray<ITargetPlatformControls*> PlatformControls;
+
 protected:
 
 	/**
 	 * This is where each platform module will fill out an array
 	*/
 	virtual void GetTargetPlatforms(TArray<ITargetPlatform*>& TargetPlatforms) = 0;
-
+	virtual void GetTargetPlatforms(TArray<ITargetPlatform*>& TargetPlatforms, TArray<ITargetPlatformSettings*> TargetPlatformSettings, TArray<ITargetPlatformControls*> TargetPlatformControls)
+	{
+		GetTargetPlatforms(TargetPlatforms);
+	}
 private:
 	/** Holds the target platforms. */
 	TArray<ITargetPlatform*> AllTargetPlatforms;
diff --git a/Engine/Source/Developer/TargetPlatform/Public/Interfaces/ITargetPlatformSettings.h b/Engine/Source/Developer/TargetPlatform/Public/Interfaces/ITargetPlatformSettings.h
new file mode 100644
--- /dev/null
+++ b/Engine/Source/Developer/TargetPlatform/Public/Interfaces/ITargetPlatformSettings.h
@@ -0,0 +1,272 @@
+// Copyright Epic Games, Inc. All Rights Reserved.
+
+#pragma once
+#include "PlatformInfo.h"
+
+enum class EOfflineBVHMode
+{
+	Disabled,
+	MaximizePerformance,
+	MinimizeMemory,
+};
+
+/**
+ * Enumerates features that may be supported by target platforms.
+ */
+enum class ETargetPlatformFeatures
+{
+	/** Audio Streaming */
+	AudioStreaming,
+
+	/** Distance field shadows. */
+	DistanceFieldShadows,
+
+	/** Distance field AO. */
+	DistanceFieldAO,
+
+	/** Gray scale SRGB texture formats support. */
+	GrayscaleSRGB,
+
+	/** High quality light maps. */
+	HighQualityLightmaps,
+
+	/** Low quality light maps. */
+	LowQualityLightmaps,
+
+	/** Run multiple game instances on a single device. */
+	MultipleGameInstances,
+
+	/** Builds can be packaged for this platform. */
+	Packaging,
+
+	/** Connect and disconnect devices through the SDK. */
+	SdkConnectDisconnect,
+
+	/** Texture streaming. */
+	TextureStreaming,
+
+	/** Mesh LOD streaming. */
+	MeshLODStreaming,
+
+	/** Landscape visual mesh LOD streaming. */
+	LandscapeMeshLODStreaming UE_DEPRECATED(5.1, "LandscapeMeshLODStreaming is now deprecated and will be removed."),
+
+	/** User credentials are required to use the device. */
+	UserCredentials,
+
+	/** The platform uses the mobile forward pipeline */
+	MobileRendering,
+
+	/** The platform uses the deferred pipeline, typically PC/Console platforms */
+	DeferredRendering,
+
+	/* Should split paks into smaller sized paks */
+	ShouldSplitPaksIntoSmallerSizes,
+
+	/* The platform supports half float vertex format */
+	HalfFloatVertexFormat,
+
+	/* The platform supports the experimental Device Output Log window */
+	DeviceOutputLog,
+
+	/* The platform supports memory mapped files */
+	MemoryMappedFiles,
+
+	/* The platform supports memory mapped audio */
+	MemoryMappedAudio,
+
+	/* The platform supports memory mapped animation */
+	MemoryMappedAnimation,
+
+	/* The platform supports sparse textures */
+	SparseTextures,
+
+	/* Can we use the virtual texture streaming system on this platform. */
+	VirtualTextureStreaming,
+
+	/** Lumen Global Illumination. */
+	LumenGI,
+
+	/** The platform supports hardware LZ decompression */
+	HardwareLZDecompression,
+
+	/* The platform makes use of extra cook-time file region metadata in its packaging process. */
+	CookFileRegionMetadata,
+
+	/** The platform supports communication (reading and writing data) between a target a connected PC. */
+	DirectDataExchange,
+
+	/** The platform supports Luminance + Alpha encoding mode for normalmaps */
+	NormalmapLAEncodingMode,
+
+	/** All devices of this platform should be grouped under one platform group */
+	ShowAsPlatformGroup,
+
+	/** Does the platform allow various connection types to be used (ie: wifi and usb) */
+	SupportsMultipleConnectionTypes,
+	
+	/** The platform can cook packages (e.g. CookedCooker) */
+	CanCookPackages
+};
+
+class ITargetPlatformSettings
+{
+public:
+	/**
+	 * Returns the config system object usable by this TargetPlatform. It should not be modified in anyway
+	 */
+	virtual FConfigCacheIni* GetConfigSystem() const = 0;
+
+	/**
+	 * Gets the platform's INI name (so an offline tool can load the INI for the given target platform).
+	 *
+	 * @see PlatformName
+	 */
+	virtual FString IniPlatformName() const = 0;
+
+	/**
+	 * Checks whether the target platform supports the specified value for the specified type of support
+	 *
+	 * @param SupportedType The type of support being queried
+	 * @param RequiredSupportedValue The value of support needed
+	 * @return true if the feature is supported, false otherwise.
+	 */
+	virtual bool SupportsValueForType(FName SupportedType, FName RequiredSupportedValue) const = 0;
+
+	/**
+	 * Gets whether the platform should use forward shading or not.
+	 */
+	virtual bool UsesForwardShading() const = 0;
+
+	/**
+	* Gets whether the platform should use DBuffer for decals.
+	*/
+	virtual bool UsesDBuffer() const = 0;
+
+	/**
+	* Gets whether the platform should output velocity in the base pass.
+	*/
+	virtual bool UsesBasePassVelocity() const = 0;
+
+	/**
+	* Gets whether the platform will use selective outputs in the base pass shaders.
+	*/
+	virtual bool UsesSelectiveBasePassOutputs() const = 0;
+
+	/**
+	* Gets whether the platform will use distance fields.
+	*/
+	virtual bool UsesDistanceFields() const = 0;
+
+	/**
+	* Gets whether the platform will use ray tracing.
+	*/
+	virtual bool UsesRayTracing() const = 0;
+
+	/**
+	* Gets static mesh offline BVH mode
+	*/
+	virtual EOfflineBVHMode GetStaticMeshOfflineBVHMode() const = 0;
+
+	/**
+	* Gets whether the platform will use compression for static mesh offline BVH.
+	*/
+	virtual bool GetStaticMeshOfflineBVHCompression() const = 0;
+
+	/**
+	* Gets whether the platform will use SH2 instead of SH3 for sky irradiance.
+	*/
+	virtual bool ForcesSimpleSkyDiffuse() const = 0;
+
+	/**
+	* Gets whether the platform will encode depth velocity.
+	*/
+	virtual bool VelocityEncodeDepth() const = 0;
+
+	/**
+	* Gets down sample mesh distance field divider.
+	*
+	* @return 1 if platform does not need to downsample mesh distance fields
+	*/
+	virtual float GetDownSampleMeshDistanceFieldDivider() const = 0;
+
+	/**
+	* Gets an integer representing the height fog mode for opaque materials on a platform.
+	* @return 0 if no override (i.e. use r.VertexFoggingForOpaque from project settings); 1 if pixel fog; 2 if vertex fog.
+	*/
+	virtual int32 GetHeightFogModeForOpaque() const = 0;
+
+	/**
+	 * Gets whether the platform uses Mobile AO
+	 */
+	virtual bool UsesMobileAmbientOcclusion() const = 0;
+
+	/**
+	 * Gets whether the platform should use DBuffer for decals when using the mobile renderer.
+	 */
+	virtual bool UsesMobileDBuffer() const = 0;
+
+	/**
+	 * Gets whether the platform uses ASTC HDR
+	 */
+	virtual bool UsesASTCHDR() const = 0;
+
+	/**
+	* Gets the shader formats this platform can use.
+	*
+	* @param OutFormats Will contain the shader formats.
+	*/
+	virtual void GetAllPossibleShaderFormats(TArray<FName>& OutFormats) const = 0;
+
+	/**
+	* Gets the shader formats that have been selected for this target platform
+	*
+	* @param OutFormats Will contain the shader formats.
+	*/
+	virtual void GetAllTargetedShaderFormats(TArray<FName>& OutFormats) const = 0;
+
+	/**
+	* Gets the shader formats that support ray tracing for this target platform.
+	*
+	* @param OutFormats Will contain the shader formats.
+	*/
+	virtual void GetRayTracingShaderFormats(TArray<FName>& OutFormats) const = 0;
+
+	/**
+	 * Checks whether the target platform supports the specified feature.
+	 *
+	 * @param Feature The feature to check.
+	 * @return true if the feature is supported, false otherwise.
+	 */
+	virtual bool SupportsFeature(ETargetPlatformFeatures Feature) const = 0;
+
+#if WITH_ENGINE
+	/**
+	 * Gets the reflection capture formats this platform needs.
+	 *
+	 * @param OutFormats Will contain the collection of formats.
+	 */
+	virtual void GetReflectionCaptureFormats(TArray<FName>& OutFormats) const = 0;
+
+	/**
+	 * Gets the texture LOD settings used by this platform.
+	 *
+	 * @return A texture LOD settings structure.
+	 */
+	virtual const class UTextureLODSettings& GetTextureLODSettings() const = 0;
+
+	/**
+	* Register Basic LOD Settings for this platform
+	*/
+	virtual void RegisterTextureLODSettings(const class UTextureLODSettings* InTextureLODSettings) = 0;
+
+	/**
+	 * Gets the static mesh LOD settings used by this platform.
+	 *
+	 * @return A static mesh LOD settings structure.
+	 */
+	virtual const class FStaticMeshLODSettings& GetStaticMeshLODSettings() const = 0;
+#endif
+	/** Virtual destructor. */
+	virtual ~ITargetPlatformSettings() { }
+};
\ No newline at end of file
diff --git a/Engine/Source/Developer/TargetPlatform/Public/Interfaces/ITargetPlatformSettingsModule.h b/Engine/Source/Developer/TargetPlatform/Public/Interfaces/ITargetPlatformSettingsModule.h
new file mode 100644
--- /dev/null
+++ b/Engine/Source/Developer/TargetPlatform/Public/Interfaces/ITargetPlatformSettingsModule.h
@@ -0,0 +1,54 @@
+// Copyright Epic Games, Inc. All Rights Reserved.
+
+#pragma once
+
+#include "CoreMinimal.h"
+#include "Modules/ModuleInterface.h"
+#include "Interfaces/ITargetPlatformSettings.h"
+
+/**
+ * Interface for target platform modules.
+ */
+class ITargetPlatformSettingsModule
+	: public IModuleInterface
+{
+
+public:
+
+	/** Virtual destructor. */
+	virtual ~ITargetPlatformSettingsModule()
+	{
+		for (ITargetPlatformSettings* TP : AllTargetPlatforms)
+		{
+			delete TP;
+		}
+		AllTargetPlatforms.Empty();
+	}
+
+	/**
+	 * Gets the module's target platforms. This should be overridden by each platform, but
+	 * currently, we are re-using the single internal GetTargetPlatform method the old TPModules will implement
+	 *
+	 * @return The target platform.
+	 */
+	TArray<ITargetPlatformSettings*> GetTargetPlatformSettings()
+	{
+		if (AllTargetPlatforms.Num() == 0)
+		{
+			GetTargetPlatformSettings(AllTargetPlatforms);
+		}
+
+		return AllTargetPlatforms;
+	}
+
+protected:
+
+	/**
+	 * This is where each platform module will fill out an array
+	*/
+	virtual void GetTargetPlatformSettings(TArray<ITargetPlatformSettings*>& TargetPlatforms) = 0;
+
+private:
+	/** Holds the target platforms. */
+	TArray<ITargetPlatformSettings*> AllTargetPlatforms;
+};
diff --git a/Engine/Source/Programs/IncludeTool/IncludeTool/Rules.cs b/Engine/Source/Programs/IncludeTool/IncludeTool/Rules.cs
--- a/Engine/Source/Programs/IncludeTool/IncludeTool/Rules.cs
+++ b/Engine/Source/Programs/IncludeTool/IncludeTool/Rules.cs
@@ -356,9 +356,8 @@ namespace IncludeTool
 			"/Engine/Source/Runtime/Online/SSL/Private/PlatformSslCertificateManager.h",
 
 			// Weird Android multiple target platform through INL file stuff
-			"/Engine/Source/Developer/Android/AndroidTargetPlatform/Private/AndroidTargetDevice.h",
-			"/Engine/Source/Developer/Android/AndroidTargetPlatform/Private/AndroidTargetDeviceOutput.h",
-			"/Engine/Source/Developer/Android/AndroidTargetPlatform/Private/AndroidTargetPlatform.h",
+			"/Engine/Source/Developer/Android/AndroidTargetPlatformControls/Private/AndroidTargetDevice.h",
+			"/Engine/Source/Developer/Android/AndroidTargetPlatformControls/Private/AndroidTargetDeviceOutput.h",
 
 			// Platform specific
 			"/Engine/Source/Runtime/Slate/Public/Framework/Text/GenericPlatformTextField.h",
diff --git a/Engine/Source/Programs/UnrealBuildTool/Platform/Android/UEBuildAndroid.cs b/Engine/Source/Programs/UnrealBuildTool/Platform/Android/UEBuildAndroid.cs
--- a/Engine/Source/Programs/UnrealBuildTool/Platform/Android/UEBuildAndroid.cs
+++ b/Engine/Source/Programs/UnrealBuildTool/Platform/Android/UEBuildAndroid.cs
@@ -341,6 +341,8 @@ namespace UnrealBuildTool
 						if (Target.bBuildDeveloperTools)
 						{
 							Rules.DynamicallyLoadedModuleNames.Add("AndroidTargetPlatform");
+							Rules.DynamicallyLoadedModuleNames.Add("AndroidTargetPlatformSettings");
+							Rules.DynamicallyLoadedModuleNames.Add("AndroidTargetPlatformControls");
 						}
 					}
 					else if (ModuleName == "TargetPlatform")
@@ -362,6 +364,8 @@ namespace UnrealBuildTool
 					if (Target.bForceBuildTargetPlatforms)
 					{
 						Rules.DynamicallyLoadedModuleNames.Add("AndroidTargetPlatform");
+						Rules.DynamicallyLoadedModuleNames.Add("AndroidTargetPlatformSettings");
+						Rules.DynamicallyLoadedModuleNames.Add("AndroidTargetPlatformControls");
 					}
 
 					if (bBuildShaderFormats)
diff --git a/Engine/Source/Runtime/Android/AndroidRuntimeSettings/AndroidRuntimeSettings.Build.cs b/Engine/Source/Runtime/Android/AndroidRuntimeSettings/AndroidRuntimeSettings.Build.cs
--- a/Engine/Source/Runtime/Android/AndroidRuntimeSettings/AndroidRuntimeSettings.Build.cs
+++ b/Engine/Source/Runtime/Android/AndroidRuntimeSettings/AndroidRuntimeSettings.Build.cs
@@ -29,6 +29,7 @@ public class AndroidRuntimeSettings : ModuleRules
             );
 
             PrivateIncludePathModuleNames.Add("AndroidTargetPlatform");
+            PrivateIncludePathModuleNames.Add("AndroidTargetPlatformControls");
         }
 	}
 }
diff --git a/Engine/Source/Runtime/Android/AndroidRuntimeSettings/Private/AndroidRuntimeSettings.cpp b/Engine/Source/Runtime/Android/AndroidRuntimeSettings/Private/AndroidRuntimeSettings.cpp
--- a/Engine/Source/Runtime/Android/AndroidRuntimeSettings/Private/AndroidRuntimeSettings.cpp
+++ b/Engine/Source/Runtime/Android/AndroidRuntimeSettings/Private/AndroidRuntimeSettings.cpp
@@ -13,6 +13,7 @@
 
 #if WITH_EDITOR
 #include "IAndroidTargetPlatformModule.h"
+#include "IAndroidTargetPlatformControlsModule.h"
 #endif
 
 DEFINE_LOG_CATEGORY(LogAndroidRuntimeSettings);
@@ -218,7 +219,7 @@ void UAndroidRuntimeSettings::PostEditChangeProperty(struct FPropertyChangedEven
 		}
 
 		// Notify the AndroidTargetPlatform module if it's loaded
-		IAndroidTargetPlatformModule* Module = FModuleManager::GetModulePtr<IAndroidTargetPlatformModule>("AndroidTargetPlatform");
+		IAndroidTargetPlatformControlsModule* Module = FModuleManager::GetModulePtr<IAndroidTargetPlatformControlsModule>("AndroidTargetPlatformControls");
 		if (Module)
 		{
 			Module->NotifyMultiSelectedFormatsChanged();
@@ -230,7 +231,7 @@ void UAndroidRuntimeSettings::PostEditChangeProperty(struct FPropertyChangedEven
 		UpdateSinglePropertyInConfigFile(PropertyChangedEvent.Property, GetDefaultConfigFilename());
 
 		// Notify the AndroidTargetPlatform module if it's loaded
-		IAndroidTargetPlatformModule* Module = FModuleManager::GetModulePtr<IAndroidTargetPlatformModule>("AndroidTargetPlatform");
+		IAndroidTargetPlatformControlsModule* Module = FModuleManager::GetModulePtr<IAndroidTargetPlatformControlsModule>("AndroidTargetPlatformControls");
 		if (Module)
 		{
 			Module->NotifyMultiSelectedFormatsChanged();
diff --git a/Engine/Source/Runtime/Engine/Classes/Engine/Texture.h b/Engine/Source/Runtime/Engine/Classes/Engine/Texture.h
--- a/Engine/Source/Runtime/Engine/Classes/Engine/Texture.h
+++ b/Engine/Source/Runtime/Engine/Classes/Engine/Texture.h
@@ -1649,6 +1649,7 @@ public:
 	 *   accounting for LODBias and other constraints
 	 */
 	ENGINE_API void GetBuiltTextureSize( const ITargetPlatform* TargetPlatform , int32 & OutSizeX, int32 & OutSizeY ) const;
+	ENGINE_API void GetBuiltTextureSize( const class ITargetPlatformSettings* TargetPlatformSettings, const class ITargetPlatformControls* TargetPlatformControls, int32 & OutSizeX, int32 & OutSizeY ) const;
 
 	/**
 	 * Serializes cooked platform data.
diff --git a/Engine/Source/Runtime/Engine/Private/StaticMesh.cpp b/Engine/Source/Runtime/Engine/Private/StaticMesh.cpp
--- a/Engine/Source/Runtime/Engine/Private/StaticMesh.cpp
+++ b/Engine/Source/Runtime/Engine/Private/StaticMesh.cpp
@@ -2178,8 +2178,7 @@ void FStaticMeshRenderData::SyncUVChannelData(const TArray<FStaticMaterial>& Obj
 /*------------------------------------------------------------------------------
 	FStaticMeshLODSettings
 ------------------------------------------------------------------------------*/
-
-void FStaticMeshLODSettings::Initialize(const ITargetPlatform* TargetPlatform)
+void FStaticMeshLODSettings::Initialize(const ITargetPlatformSettings* TargetPlatformSettings)
 {
 	check(!Groups.Num());
 	// Ensure there is a default LOD group.
@@ -2187,7 +2186,7 @@ void FStaticMeshLODSettings::Initialize(const ITargetPlatform* TargetPlatform)
 
 	// Read individual entries from a config file.
 	const TCHAR* IniSection = TEXT("StaticMeshLODSettings");
-	const FConfigSection* Section = TargetPlatform->GetConfigSystem()->GetSection(IniSection, false, GEngineIni);
+	const FConfigSection* Section = TargetPlatformSettings->GetConfigSystem()->GetSection(IniSection, false, GEngineIni);
 	if (Section)
 	{
 		for (TMultiMap<FName,FConfigValue>::TConstIterator It(*Section); It; ++It)
@@ -2246,6 +2245,10 @@ void FStaticMeshLODSettings::Initialize(const ITargetPlatform* TargetPlatform)
 		}
 	}
 }
+void FStaticMeshLODSettings::Initialize(const ITargetPlatform* TargetPlatform)
+{
+	Initialize(&TargetPlatform->GetPlatformSettings());
+}
 
 void FStaticMeshLODSettings::ReadEntry(FStaticMeshLODGroup& Group, FString Entry)
 {
diff --git a/Engine/Source/Runtime/Engine/Private/Texture.cpp b/Engine/Source/Runtime/Engine/Private/Texture.cpp
--- a/Engine/Source/Runtime/Engine/Private/Texture.cpp
+++ b/Engine/Source/Runtime/Engine/Private/Texture.cpp
@@ -3458,7 +3458,7 @@ static FName CachedGetLatestOodleSdkVersion()
 	return Once.Value;
 }
 
-static FName ConditionalGetPrefixedFormat(FName TextureFormatName, const ITargetPlatform* TargetPlatform, bool bOodleTextureSdkVersionIsNone)
+static FName ConditionalGetPrefixedFormat(FName TextureFormatName, const ITargetPlatformSettings* TargetPlatformSettings, bool bOodleTextureSdkVersionIsNone)
 {
 #if WITH_EDITOR
 
@@ -3470,7 +3470,7 @@ static FName ConditionalGetPrefixedFormat(FName TextureFormatName, const ITarget
 	// TextureCompressionFormat is required, TextureCompressionFormatWithVersion is optional
 
 	FString TextureCompressionFormat;
-	bool bHasFormat = TargetPlatform->GetConfigSystem()->GetString(TEXT("AlternateTextureCompression"), TEXT("TextureCompressionFormat"), TextureCompressionFormat, GEngineIni);
+	bool bHasFormat = TargetPlatformSettings->GetConfigSystem()->GetString(TEXT("AlternateTextureCompression"), TEXT("TextureCompressionFormat"), TextureCompressionFormat, GEngineIni);
 	bHasFormat = bHasFormat && ! TextureCompressionFormat.IsEmpty();
 	
 	if ( bHasFormat )
@@ -3479,7 +3479,7 @@ static FName ConditionalGetPrefixedFormat(FName TextureFormatName, const ITarget
 		{
 			//	new (optional) pref : TextureCompressionFormatWithVersion
 			FString TextureCompressionFormatWithVersion;
-			bool bHasFormatWithVersion = TargetPlatform->GetConfigSystem()->GetString(TEXT("AlternateTextureCompression"), TEXT("TextureCompressionFormatWithVersion"), TextureCompressionFormatWithVersion, GEngineIni);
+			bool bHasFormatWithVersion = TargetPlatformSettings->GetConfigSystem()->GetString(TEXT("AlternateTextureCompression"), TEXT("TextureCompressionFormatWithVersion"), TextureCompressionFormatWithVersion, GEngineIni);
 			bHasFormatWithVersion = bHasFormatWithVersion && ! TextureCompressionFormatWithVersion.IsEmpty();
 			if ( bHasFormatWithVersion )
 			{
@@ -3540,8 +3540,12 @@ static FName ConditionalGetPrefixedFormat(FName TextureFormatName, const ITarget
 
 	return TextureFormatName;
 }
+static FName ConditionalGetPrefixedFormat(FName TextureFormatName, const ITargetPlatform* TargetPlatform, bool bOodleTextureSdkVersionIsNone)
+{
+	return ConditionalGetPrefixedFormat(TextureFormatName, &TargetPlatform->GetPlatformSettings(), bOodleTextureSdkVersionIsNone);
+}
 
-void UTexture::GetBuiltTextureSize( const ITargetPlatform* TargetPlatform , int32 & OutSizeX, int32 & OutSizeY ) const
+void UTexture::GetBuiltTextureSize(const ITargetPlatformSettings* TargetPlatformSettings, const ITargetPlatformControls* TargetPlatformControls, int32 & OutSizeX, int32 & OutSizeY ) const
 {
 	// @todo Oodle : SizeZ
 	// @todo Oodle : verify against TextureCompressorModule
@@ -3601,9 +3605,9 @@ void UTexture::GetBuiltTextureSize( const ITargetPlatform* TargetPlatform , int3
 	}
 
 	static const auto CVarVirtualTexturesEnabled = IConsoleManager::Get().FindTConsoleVariableDataInt(TEXT("r.VirtualTextures")); check(CVarVirtualTexturesEnabled);
-	const bool bVirtualTextureStreaming = CVarVirtualTexturesEnabled->GetValueOnAnyThread() && TargetPlatform->SupportsFeature(ETargetPlatformFeatures::VirtualTextureStreaming) && VirtualTextureStreaming;
+	const bool bVirtualTextureStreaming = CVarVirtualTexturesEnabled->GetValueOnAnyThread() && TargetPlatformSettings->SupportsFeature(ETargetPlatformFeatures::VirtualTextureStreaming) && VirtualTextureStreaming;
 
-	const UTextureLODSettings& LODSettings = TargetPlatform->GetTextureLODSettings();
+	const UTextureLODSettings& LODSettings = TargetPlatformSettings->GetTextureLODSettings();
  	const uint32 LODBiasNoCinematics = FMath::Max<int32>(LODSettings.CalculateLODBias(SizeX, SizeY, MaxTextureSize, LODGroup, LODBias, 0, MipGenSettings, bVirtualTextureStreaming), 0);
 	SizeX = FMath::Max<int32>(SizeX >> LODBiasNoCinematics, 1);
 	SizeY = FMath::Max<int32>(SizeY >> LODBiasNoCinematics, 1);
@@ -3636,10 +3640,13 @@ void UTexture::GetBuiltTextureSize( const ITargetPlatform* TargetPlatform , int3
 	OutSizeX = SizeX;
 	OutSizeY = SizeY;
 }
-
+void UTexture::GetBuiltTextureSize(const ITargetPlatform* TargetPlatform, int32& OutSizeX, int32& OutSizeY) const
+{
+	return GetBuiltTextureSize(&TargetPlatform->GetPlatformSettings(), &TargetPlatform->GetPlatformControls(), OutSizeX, OutSizeY);
+}
 // this should not be called directly; it is called from TargetPlatform GetTextureFormats
 //	entry point API is GetPlatformTextureFormatNamesWithPrefix
-FName GetDefaultTextureFormatName( const ITargetPlatform* TargetPlatform, const UTexture* Texture, int32 LayerIndex, 
+FName GetDefaultTextureFormatName( const ITargetPlatformSettings* TargetPlatformSettings, const ITargetPlatformControls* TargetPlatformControls, const UTexture* Texture, int32 LayerIndex, 
 	bool bSupportCompressedVolumeTexture, int32 Unused_BlockSize, bool bSupportFilteredFloat32Textures )
 {
 	FName TextureFormatName = NAME_None;
@@ -3671,7 +3678,7 @@ FName GetDefaultTextureFormatName( const ITargetPlatform* TargetPlatform, const
 	static FName NameA1RGB555(TEXT("A1RGB555"));
 	
 
-	check(TargetPlatform);
+	check(TargetPlatformSettings);
 
 	FTextureFormatSettings FormatSettings;
 	Texture->GetLayerFormatSettings(LayerIndex, FormatSettings);
@@ -3682,7 +3689,7 @@ FName GetDefaultTextureFormatName( const ITargetPlatform* TargetPlatform, const
 
 	// see if compression needs to be forced off even if requested :
 	bool bNoCompression = FormatSettings.CompressionNone				// Code wants the texture uncompressed.
-		|| (TargetPlatform->HasEditorOnlyData() && Texture->DeferCompression)	// The user wishes to defer compression, this is ok for the Editor only.
+		|| (TargetPlatformControls->HasEditorOnlyData() && Texture->DeferCompression)	// The user wishes to defer compression, this is ok for the Editor only.
 		|| (FormatSettings.CompressionSettings == TC_EditorIcon) // TC_EditorIcon is "UserInterface2D"
 		|| (Texture->LODGroup == TEXTUREGROUP_ColorLookupTable)	// Textures in certain LOD groups should remain uncompressed.
 		|| (Texture->LODGroup == TEXTUREGROUP_Bokeh)
@@ -3693,7 +3700,7 @@ FName GetDefaultTextureFormatName( const ITargetPlatform* TargetPlatform, const
 	if (!bNoCompression)
 	{
 		int32 SizeX,SizeY;
-		Texture->GetBuiltTextureSize(TargetPlatform,SizeX,SizeY);
+		Texture->GetBuiltTextureSize(TargetPlatformSettings, TargetPlatformControls,SizeX,SizeY);
 
 	
 		// Don't compress textures smaller than the DXT block size.
@@ -3712,7 +3719,7 @@ FName GetDefaultTextureFormatName( const ITargetPlatform* TargetPlatform, const
 
 	FString UseDXT5NormalMapsString;
 
-	if (TargetPlatform->GetConfigSystem()->GetString(TEXT("SystemSettings"), TEXT("Compat.UseDXT5NormalMaps"), UseDXT5NormalMapsString, GEngineIni))
+	if (TargetPlatformSettings->GetConfigSystem()->GetString(TEXT("SystemSettings"), TEXT("Compat.UseDXT5NormalMaps"), UseDXT5NormalMapsString, GEngineIni))
 	{
 		bUseDXT5NormalMap = FCString::ToBool(*UseDXT5NormalMapsString);
 	}
@@ -3721,7 +3728,7 @@ FName GetDefaultTextureFormatName( const ITargetPlatform* TargetPlatform, const
 
 	if (FormatSettings.CompressionSettings == TC_LQ)
 	{
-		bool bLQCompressionSupported = TargetPlatform->SupportsLQCompressionTextureFormat();
+		bool bLQCompressionSupported = TargetPlatformControls->SupportsLQCompressionTextureFormat();
 		if(bLQCompressionSupported)
 		{
 			TextureFormatName = FormatSettings.CompressionNoAlpha ? NameR5G6B5 : NameA1RGB555;
@@ -3882,7 +3889,7 @@ FName GetDefaultTextureFormatName( const ITargetPlatform* TargetPlatform, const
 	// note: GrayscaleSRGB is off on all targetplatforms currently
 	// someday: I think this could use G16 instead and be half the size
 	//	 (that's doing the gamma->linear in the G8->G16 conversion)
-	if ((TextureFormatName == NameG8) && FormatSettings.SRGB && !TargetPlatform->SupportsFeature(ETargetPlatformFeatures::GrayscaleSRGB))
+	if ((TextureFormatName == NameG8) && FormatSettings.SRGB && !TargetPlatformSettings->SupportsFeature(ETargetPlatformFeatures::GrayscaleSRGB))
 	{
 		TextureFormatName = NameBGRA8;
 	}
@@ -3893,7 +3900,7 @@ FName GetDefaultTextureFormatName( const ITargetPlatform* TargetPlatform, const
 	{
 		// Texture::Filter can be manually set to TF_Nearest , if it's Default it comes from LOD Group
 		//   eg. Nearest for TEXTUREGROUP_ColorLookupTable and TEXTUREGROUP_Pixels2D
-		const UTextureLODSettings& LODSettings = TargetPlatform->GetTextureLODSettings();
+		const UTextureLODSettings& LODSettings = TargetPlatformSettings->GetTextureLODSettings();
 		ETextureSamplerFilter Filter = LODSettings.GetSamplerFilter(Texture);
 
 		if ( Filter != ETextureSamplerFilter::Point )
@@ -3928,6 +3935,11 @@ FName GetDefaultTextureFormatName( const ITargetPlatform* TargetPlatform, const
 
 	return TextureFormatName;
 }
+FName GetDefaultTextureFormatName(const ITargetPlatform* TargetPlatform, const UTexture* Texture, int32 LayerIndex,
+	bool bSupportCompressedVolumeTexture, int32 Unused_BlockSize, bool bSupportFilteredFloat32Textures)
+{
+	return GetDefaultTextureFormatName(&TargetPlatform->GetPlatformSettings(), &TargetPlatform->GetPlatformControls(), Texture, LayerIndex, bSupportCompressedVolumeTexture, Unused_BlockSize, bSupportFilteredFloat32Textures);
+}
 
 #if WITH_EDITOR
 // this is the entry point API to get the texture format name
@@ -3994,19 +4006,24 @@ void UTexture::GetPlatformTextureFormatNamesWithPrefix(const class ITargetPlatfo
 }
 #endif
 
-void GetDefaultTextureFormatNamePerLayer(TArray<FName>& OutFormatNames, const class ITargetPlatform* TargetPlatform, const class UTexture* Texture, 
+void GetDefaultTextureFormatNamePerLayer(TArray<FName>& OutFormatNames, const class ITargetPlatformSettings* TargetPlatformSettings, const class ITargetPlatformControls* TargetPlatformControls, const class UTexture* Texture,
 	bool bSupportCompressedVolumeTexture, int32 Unused_BlockSize, bool bSupportFilteredFloat32Textures )
 {
 #if WITH_EDITOR
 	OutFormatNames.Reserve(Texture->Source.GetNumLayers());
 	for (int32 LayerIndex = 0; LayerIndex < Texture->Source.GetNumLayers(); ++LayerIndex)
 	{
-		OutFormatNames.Add(GetDefaultTextureFormatName(TargetPlatform, Texture, LayerIndex, bSupportCompressedVolumeTexture, Unused_BlockSize, bSupportFilteredFloat32Textures));
+		OutFormatNames.Add(GetDefaultTextureFormatName(TargetPlatformSettings, TargetPlatformControls, Texture, LayerIndex, bSupportCompressedVolumeTexture, Unused_BlockSize, bSupportFilteredFloat32Textures));
 	}
 #endif // WITH_EDITOR
 }
+void GetDefaultTextureFormatNamePerLayer(TArray<FName>& OutFormatNames, const class ITargetPlatform* TargetPlatform, const class UTexture* Texture,
+	bool bSupportCompressedVolumeTexture, int32 Unused_BlockSize, bool bSupportFilteredFloat32Textures)
+{
+	GetDefaultTextureFormatNamePerLayer(OutFormatNames, &TargetPlatform->GetPlatformSettings(), &TargetPlatform->GetPlatformControls(), Texture, bSupportCompressedVolumeTexture, Unused_BlockSize, bSupportFilteredFloat32Textures);
+}
 
-void GetAllDefaultTextureFormats(const class ITargetPlatform* TargetPlatform, TArray<FName>& OutFormats)
+void GetAllDefaultTextureFormats(const class ITargetPlatformSettings* TargetPlatformSettings, TArray<FName>& OutFormats)
 {
 	// this is only used by CookOnTheFlyServer, it could be removed entirely
 
@@ -4054,8 +4071,8 @@ void GetAllDefaultTextureFormats(const class ITargetPlatform* TargetPlatform, TA
 	int NumBaseFormats = OutFormats.Num();
 	for (int Index = 0; Index < NumBaseFormats; Index++)
 	{
-		OutFormats.Add(ConditionalGetPrefixedFormat(OutFormats[Index], TargetPlatform, true));
-		OutFormats.Add(ConditionalGetPrefixedFormat(OutFormats[Index], TargetPlatform, false));
+		OutFormats.Add(ConditionalGetPrefixedFormat(OutFormats[Index], TargetPlatformSettings, true));
+		OutFormats.Add(ConditionalGetPrefixedFormat(OutFormats[Index], TargetPlatformSettings, false));
 	}
 	
 	// make unique:		
@@ -4064,6 +4081,11 @@ void GetAllDefaultTextureFormats(const class ITargetPlatform* TargetPlatform, TA
 #endif
 }
 
+void GetAllDefaultTextureFormats(const class ITargetPlatform* TargetPlatform, TArray<FName>& OutFormats)
+{
+	GetAllDefaultTextureFormats(&TargetPlatform->GetPlatformSettings(), OutFormats);
+}
+
 #if WITH_EDITOR
 
 void UTexture::NotifyMaterials(const ENotifyMaterialsEffectOnShaders EffectOnShaders)
diff --git a/Engine/Source/Runtime/Engine/Public/StaticMeshResources.h b/Engine/Source/Runtime/Engine/Public/StaticMeshResources.h
--- a/Engine/Source/Runtime/Engine/Public/StaticMeshResources.h
+++ b/Engine/Source/Runtime/Engine/Public/StaticMeshResources.h
@@ -159,6 +159,7 @@ public:
 	 * @param IniFile Preloaded ini file object to load from
 	 */
 	ENGINE_API void Initialize(const ITargetPlatform* TargetPlatform);
+	ENGINE_API void Initialize(const class ITargetPlatformSettings* TargetPlatform);
 
 	/** Retrieve the settings for the specified LOD group. */
 	const FStaticMeshLODGroup& GetLODGroup(FName LODGroup) const
diff --git a/Engine/Source/Runtime/Engine/Public/TextureResource.h b/Engine/Source/Runtime/Engine/Public/TextureResource.h
--- a/Engine/Source/Runtime/Engine/Public/TextureResource.h
+++ b/Engine/Source/Runtime/Engine/Public/TextureResource.h
@@ -656,9 +656,12 @@ private:
 
 /** Gets the name of a format for the given LayerIndex */
 ENGINE_API FName GetDefaultTextureFormatName( const class ITargetPlatform* TargetPlatform, const class UTexture* Texture, int32 LayerIndex, bool bSupportCompressedVolumeTexture, int32 Unused_BlockSize, bool bSupportFilteredFloat32Textures);
+ENGINE_API FName GetDefaultTextureFormatName( const class ITargetPlatformSettings* TargetPlatformSettings, const class ITargetPlatformControls* TargetPlatformControls, const class UTexture* Texture, int32 LayerIndex, bool bSupportCompressedVolumeTexture, int32 Unused_BlockSize, bool bSupportFilteredFloat32Textures);
 
 /** Gets an array of format names for each layer in the texture */
 ENGINE_API void GetDefaultTextureFormatNamePerLayer(TArray<FName>& OutFormatNames, const class ITargetPlatform* TargetPlatform, const class UTexture* Texture, bool bSupportCompressedVolumeTexture, int32 Unused_BlockSize, bool bSupportFilteredFloat32Textures);
+ENGINE_API void GetDefaultTextureFormatNamePerLayer(TArray<FName>& OutFormatNames, const class ITargetPlatformSettings* TargetPlatformSettings, const class ITargetPlatformControls* TargetPlatformControls, const class UTexture* Texture, bool bSupportCompressedVolumeTexture, int32 Unused_BlockSize, bool bSupportFilteredFloat32Textures);
 
 // returns all the texture formats which can be returned by GetDefaultTextureFormatName
 ENGINE_API void GetAllDefaultTextureFormats( const class ITargetPlatform* TargetPlatform, TArray<FName>& OutFormats);
+ENGINE_API void GetAllDefaultTextureFormats( const class ITargetPlatformSettings* TargetPlatformSettings, TArray<FName>& OutFormats);
