Add missing matching #undef NAMESPACE_LOCTEXT directive

diff --git a/Engine/Source/Developer/TargetPlatform/Private/TargetPlatformControlsBase.cpp b/Engine/Source/Developer/TargetPlatform/Private/TargetPlatformControlsBase.cpp
--- a/Engine/Source/Developer/TargetPlatform/Private/TargetPlatformControlsBase.cpp
+++ b/Engine/Source/Developer/TargetPlatform/Private/TargetPlatformControlsBase.cpp
@@ -294,4 +294,6 @@ void FTargetPlatformControlsBase::GetWaveFormatModuleHints(TArray<FName>& OutMod
 	GetAudioFormatSettings().GetWaveFormatModuleHints(OutModuleNames);
 }
 
-#endif // WITH_ENGINE
\ No newline at end of file
+#endif // WITH_ENGINE
+
+#undef LOCTEXT_NAMESPACE
\ No newline at end of file
