Static Analysis fix

diff --git a/Engine/Source/Developer/TargetPlatform/Private/TargetPlatformManagerModule.cpp b/Engine/Source/Developer/TargetPlatform/Private/TargetPlatformManagerModule.cpp
--- a/Engine/Source/Developer/TargetPlatform/Private/TargetPlatformManagerModule.cpp
+++ b/Engine/Source/Developer/TargetPlatform/Private/TargetPlatformManagerModule.cpp
@@ -794,7 +794,10 @@ protected:
 			for (ITargetPlatformSettings* Platform : TargetPlatformSettings)
 			{
 				PlatformSettings.Add(Platform);
-				Module->PlatformSettings.Add(Platform);
+				if (Module != nullptr)
+				{
+					Module->PlatformSettings.Add(Platform);
+				}
 			}
 		}
 
diff --git a/Engine/Source/Developer/TargetPlatform/Public/Interfaces/ITargetDevice.h b/Engine/Source/Developer/TargetPlatform/Public/Interfaces/ITargetDevice.h
--- a/Engine/Source/Developer/TargetPlatform/Public/Interfaces/ITargetDevice.h
+++ b/Engine/Source/Developer/TargetPlatform/Public/Interfaces/ITargetDevice.h
@@ -331,8 +331,10 @@ protected:
 	// Temporary until we get rid of this.
 	virtual const class ITargetPlatform& GetTargetPlatform() const
 	{
-		ITargetPlatform* DummyReference = nullptr;
-		return *DummyReference;
+		unimplemented();
+		static ITargetPlatform* DummyReference = nullptr;
+		// Code should never reach this point, this is just temporary until we port all platforms to new system
+		return *DummyReference; //-V522
 	};
 public:
 
