Make FTargetPlatformManagerModule::ShaderFormatVersion thread safe

diff --git a/Engine/Source/Developer/TargetPlatform/Private/TargetPlatformManagerModule.cpp b/Engine/Source/Developer/TargetPlatform/Private/TargetPlatformManagerModule.cpp
--- a/Engine/Source/Developer/TargetPlatform/Private/TargetPlatformManagerModule.cpp
+++ b/Engine/Source/Developer/TargetPlatform/Private/TargetPlatformManagerModule.cpp
@@ -717,22 +717,30 @@ public:
 
 	virtual uint32 ShaderFormatVersion(FName Name) override
 	{
-		static TMap<FName, uint32> AlreadyFound;
-		uint32* Result = AlreadyFound.Find(Name);
+		static bool bInitialized = false;
+		static TMap<FName, uint32> FormatVersionCache;
 
-		if (!Result)
+		if (!bInitialized || bForceCacheUpdate)
 		{
-			const IShaderFormat* SF = FindShaderFormat(Name);
+			FormatVersionCache.Reset();
 
-			if (SF)
+			for (const IShaderFormat* SF : GetShaderFormats())
 			{
-				Result = &AlreadyFound.Add(Name, SF->GetVersion(Name));
+				TArray<FName> Formats;
+				SF->GetSupportedFormats(Formats);
+				for (FName FormatName : Formats)
+				{
+					FormatVersionCache.FindOrAdd(FormatName, SF->GetVersion(FormatName));
+				}
 			}
+			bInitialized = true;
 		}
 
+		const uint32* Result = FormatVersionCache.Find(Name);
 		if (!Result)
 		{
-			UE_LOG(LogTargetPlatformManager, Fatal, TEXT("No ShaderFormat found for %s!"), *Name.ToString());
+			UE_LOG(LogTargetPlatformManager, Fatal, TEXT("ShaderFormat not found for %s!  Dynamically loaded shader formats require invalidation of FormatVersionCache."), *Name.ToString());
+			return INDEX_NONE;
 		}
 
 		return *Result;
