Deprecated FTargetPlatformBase, TTargetPlatformBase and TNonDesktopTargetPlatformBase
Replaced with FTargetPlatformControlsBase/FTargetPlatformSettingsBase,  TTargetPlatformControlsBase/TTargetPlatformSettingsBase and TNonDesktopTargetPlatformControlsBase

diff --git a/Engine/Source/Developer/Android/AndroidTargetPlatformControls/Private/AndroidTargetPlatformControls.cpp b/Engine/Source/Developer/Android/AndroidTargetPlatformControls/Private/AndroidTargetPlatformControls.cpp
--- a/Engine/Source/Developer/Android/AndroidTargetPlatformControls/Private/AndroidTargetPlatformControls.cpp
+++ b/Engine/Source/Developer/Android/AndroidTargetPlatformControls/Private/AndroidTargetPlatformControls.cpp
@@ -51,7 +51,6 @@ class UTexture;
 class UTextureLODSettings;
 struct FAndroidDeviceInfo;
 enum class ETargetPlatformFeatures;
-template<typename TPlatformProperties> class TTargetPlatformBase;
 
 static FString GetLicensePath()
 {
diff --git a/Engine/Source/Developer/Android/AndroidTargetPlatformControls/Private/AndroidTargetPlatformControls.h b/Engine/Source/Developer/Android/AndroidTargetPlatformControls/Private/AndroidTargetPlatformControls.h
--- a/Engine/Source/Developer/Android/AndroidTargetPlatformControls/Private/AndroidTargetPlatformControls.h
+++ b/Engine/Source/Developer/Android/AndroidTargetPlatformControls/Private/AndroidTargetPlatformControls.h
@@ -30,12 +30,10 @@ class IAndroidDeviceDetection;
 class ITargetPlatform;
 class UTextureLODSettings;
 enum class ETargetPlatformFeatures;
-template<typename TPlatformProperties> class TTargetPlatformBase;
 
 template< typename InElementType, typename KeyFuncs, typename Allocator > class TSet;
 template<typename KeyType, typename ValueType, typename SetAllocator, typename KeyFuncs > class TMap;
 template<typename KeyType, typename ValueType, typename SetAllocator, typename KeyFuncs > class TMultiMap;
-template<typename TPlatformProperties> class TTargetPlatformBase;
 
 /**
  * FAndroidTargetPlatformControls, abstraction for cooking Android platforms
diff --git a/Engine/Source/Developer/TargetPlatform/Private/TargetPlatformBase.cpp b/Engine/Source/Developer/TargetPlatform/Private/TargetPlatformBase.cpp
--- a/Engine/Source/Developer/TargetPlatform/Private/TargetPlatformBase.cpp
+++ b/Engine/Source/Developer/TargetPlatform/Private/TargetPlatformBase.cpp
@@ -16,6 +16,8 @@
 
 #define LOCTEXT_NAMESPACE "TargetPlatform"
 
+PRAGMA_DISABLE_DEPRECATION_WARNINGS
+
 bool FTargetPlatformBase::UsesForwardShading() const
 {
 	static IConsoleVariable* CVarForwardShading = IConsoleManager::Get().FindConsoleVariable(TEXT("r.ForwardShading"));
@@ -409,4 +411,6 @@ void FTargetPlatformBase::GetTextureSizeLimits(uint64 & OutMaximumSurfaceBytes,
 
 #endif // WITH_ENGINE
 
+PRAGMA_ENABLE_DEPRECATION_WARNINGS
+
 #undef LOCTEXT_NAMESPACE
diff --git a/Engine/Source/Developer/TargetPlatform/Private/TargetPlatformControlsBase.cpp b/Engine/Source/Developer/TargetPlatform/Private/TargetPlatformControlsBase.cpp
--- a/Engine/Source/Developer/TargetPlatform/Private/TargetPlatformControlsBase.cpp
+++ b/Engine/Source/Developer/TargetPlatform/Private/TargetPlatformControlsBase.cpp
@@ -260,7 +260,7 @@ FTargetPlatformControlsBase::FTargetPlatformControlsBase(const PlatformInfo::FTa
 	: ITargetPlatformControls(TargetPlatformSettings)
 	, PlatformInfo(InPlatformInfo)
 {
-	checkf(PlatformInfo, TEXT("Null PlatformInfo was passed to FTargetPlatformBase. Check the static IsUsable function before creating this object. See FWindowsTargetPlatformModule::GetTargetPlatform()"));
+	checkf(PlatformInfo, TEXT("Null PlatformInfo was passed to FTargetPlatformControlsBase. Check the static IsUsable function before creating this object. See FWindowsTargetPlatformModule::GetTargetPlatform()"));
 
 	PlatformOrdinal = AssignPlatformOrdinal(*this);
 
diff --git a/Engine/Source/Developer/TargetPlatform/Public/Common/TargetPlatformBase.h b/Engine/Source/Developer/TargetPlatform/Public/Common/TargetPlatformBase.h
--- a/Engine/Source/Developer/TargetPlatform/Public/Common/TargetPlatformBase.h
+++ b/Engine/Source/Developer/TargetPlatform/Public/Common/TargetPlatformBase.h
@@ -36,7 +36,7 @@ namespace Audio
 /**
  * Base class for target platforms.
  */
-class FTargetPlatformBase
+class UE_DEPRECATED(5.5, "FTargetPlatformBase is deprecated, please use FTargetPlatformControlsBase(SDK required to build) or FTargetPlatformSettingsBase(no SDK required).") FTargetPlatformBase
 	: public ITargetPlatform
 {
 public:
@@ -317,7 +317,7 @@ private:
  * @param TPlatformProperties Type of platform properties.
  */
 template<typename TPlatformProperties>
-class TTargetPlatformBase
+class UE_DEPRECATED(5.5, "TTargetPlatformBase is deprecated, please use TTargetPlatformControlsBase(SDK required to build) or TTargetPlatformSettingsBase(no SDK required).") TTargetPlatformBase
 	: public FTargetPlatformBase
 {
 public:
@@ -549,7 +549,7 @@ public:
 
 
 template<typename TPlatformProperties>
-class TNonDesktopTargetPlatformBase 
+class UE_DEPRECATED(5.5, "TNonDesktopTargetPlatformBase is deprecated, please use TNonDesktopTargetPlatformControlsBase.") TNonDesktopTargetPlatformBase
 	: public TTargetPlatformBase<TPlatformProperties>
 {
 public:
