Add PlatformName to FDataDrivenShaderPlatformInfo.
In Editor, when using -game we now also load the TPS module of the preview platform if we are using OverrideSP command

diff --git a/Engine/Source/Developer/TargetPlatform/Private/TargetPlatformManagerModule.cpp b/Engine/Source/Developer/TargetPlatform/Private/TargetPlatformManagerModule.cpp
--- a/Engine/Source/Developer/TargetPlatform/Private/TargetPlatformManagerModule.cpp
+++ b/Engine/Source/Developer/TargetPlatform/Private/TargetPlatformManagerModule.cpp
@@ -30,6 +30,7 @@
 #include "PlatformInfo.h"
 #include "DesktopPlatformModule.h"
 #include "Interfaces/ITurnkeySupportModule.h"
+#include "DataDrivenShaderPlatformInfo.h"
 
 DEFINE_LOG_CATEGORY_STATIC(LogTargetPlatformManager, Log, All);
 
@@ -872,6 +873,18 @@ protected:
 		}
 #endif
 
+	// Get the platform we are previewing if GMaxRHIShaderPlatform is a preview SP
+#if WITH_EDITOR
+		FName PlatformNamePreview = NAME_None;
+		if (IsRunningGame())
+		{
+			if (FDataDrivenShaderPlatformInfo::GetIsPreviewPlatform(GMaxRHIShaderPlatform))
+			{
+				PlatformNamePreview = FDataDrivenShaderPlatformInfo::GetPlatformName(GMaxRHIShaderPlatform);
+			}
+		}
+#endif
+
 		// find a set of valid target platform names (the platform DataDrivenPlatformInfo.ini file was found indicates support for the platform 
 		// exists on disk, so the TP is expected to work)
 		FScopedSlowTask SlowTask((float)FDataDrivenPlatformInfoRegistry::GetAllPlatformInfos().Num());
@@ -887,10 +900,13 @@ protected:
 
 #if WITH_EDITOR
 			// if we have the editor and we are using -game
-			// only need to instantiate the current platform 
+			// only need to instantiate the current platform if we are not using OverrideSP command
+			// If we are using OverrideSP command that means our SP is a preview SP 
+			// Because of that we must also load the TPS of the Platform we are previewing.
 			if (IsRunningGame())
 			{
-				if (PlatformName != FPlatformProperties::IniPlatformName())
+				if ((PlatformName != FPlatformProperties::IniPlatformName()) && 
+					(PlatformName != PlatformNamePreview))
 				{
 					continue;
 				}
diff --git a/Engine/Source/Developer/TargetPlatform/TargetPlatform.Build.cs b/Engine/Source/Developer/TargetPlatform/TargetPlatform.Build.cs
--- a/Engine/Source/Developer/TargetPlatform/TargetPlatform.Build.cs
+++ b/Engine/Source/Developer/TargetPlatform/TargetPlatform.Build.cs
@@ -16,6 +16,7 @@ public class TargetPlatform : ModuleRules
 		PublicDependencyModuleNames.Add("AudioPlatformConfiguration");
 		PublicDependencyModuleNames.Add("DesktopPlatform");
 		PublicDependencyModuleNames.Add("Analytics");
+		PublicDependencyModuleNames.Add("RHI");
 
 		// TextureFormat contains public headers that were historically part of TargetPlatform, so it is exposed
 		// as a public include path on TargetPlatform.
diff --git a/Engine/Source/Runtime/RHI/Private/DataDrivenShaderPlatformInfo.cpp b/Engine/Source/Runtime/RHI/Private/DataDrivenShaderPlatformInfo.cpp
--- a/Engine/Source/Runtime/RHI/Private/DataDrivenShaderPlatformInfo.cpp
+++ b/Engine/Source/Runtime/RHI/Private/DataDrivenShaderPlatformInfo.cpp
@@ -409,6 +409,7 @@ void FGenericDataDrivenShaderPlatformInfo::Initialize()
 
 				// at this point, we can start pulling information out
 				Infos[ShaderPlatform].Name = *SectionName.Mid(15);
+				Infos[ShaderPlatform].PlatformName = FName(*PlatformName);
 				PlatformNameToShaderPlatformMap.FindOrAdd(Infos[ShaderPlatform].Name) = ShaderPlatform;
 				ParseDataDrivenShaderInfo(SectionSettings, ShaderPlatform);
 				Infos[ShaderPlatform].bContainsValidPlatformInfo = true;
@@ -425,6 +426,7 @@ void FGenericDataDrivenShaderPlatformInfo::Initialize()
 							const EShaderPlatform PreviewShaderPlatform = EShaderPlatform(CustomShaderPlatform++);
 							FGenericDataDrivenShaderPlatformInfo& PreviewInfo = Infos[PreviewShaderPlatform];
 							PreviewInfo.Name = Item.PreviewShaderPlatformName;
+							PreviewInfo.PlatformName = Infos[ShaderPlatform].PlatformName;
 							ParseDataDrivenShaderInfo(SectionSettings, PreviewShaderPlatform);
 							PreviewInfo.bIsPreviewPlatform = true;
 							PreviewInfo.bContainsValidPlatformInfo = true;
diff --git a/Engine/Source/Runtime/RHI/Public/DataDrivenShaderPlatformInfo.h b/Engine/Source/Runtime/RHI/Public/DataDrivenShaderPlatformInfo.h
--- a/Engine/Source/Runtime/RHI/Public/DataDrivenShaderPlatformInfo.h
+++ b/Engine/Source/Runtime/RHI/Public/DataDrivenShaderPlatformInfo.h
@@ -19,6 +19,7 @@ extern RHI_API const FName LANGUAGE_Nintendo;
 class FGenericDataDrivenShaderPlatformInfo
 {
 	FName Name;
+	FName PlatformName;
 	FName Language;
 	ERHIFeatureLevel::Type MaxFeatureLevel;
 	FName ShaderFormat;
@@ -153,6 +154,12 @@ public:
 		return Infos[Platform].Name;
 	}
 
+	static FORCEINLINE_DEBUGGABLE const FName GetPlatformName(const FStaticShaderPlatform Platform)
+	{
+		check(IsValid(Platform));
+		return Infos[Platform].PlatformName;
+	}
+
 	static FORCEINLINE_DEBUGGABLE const FName GetShaderFormat(const FStaticShaderPlatform Platform)
 	{
 		check(IsValid(Platform));
