Add Fallback Format Cvar for when ASTC HDR is not supported. You can select non-HDR ASTC format(default) or uncompressed RGBA16F

diff --git a/Engine/Source/Developer/Android/AndroidTargetPlatformControls/Private/AndroidTargetPlatformControls.cpp b/Engine/Source/Developer/Android/AndroidTargetPlatformControls/Private/AndroidTargetPlatformControls.cpp
--- a/Engine/Source/Developer/Android/AndroidTargetPlatformControls/Private/AndroidTargetPlatformControls.cpp
+++ b/Engine/Source/Developer/Android/AndroidTargetPlatformControls/Private/AndroidTargetPlatformControls.cpp
@@ -585,7 +585,7 @@ void FAndroid_ASTCTargetPlatformControls::GetTextureFormats(const UTexture* Text
 		{
 			if (TextureFormatName == AndroidTexFormat::NameASTC_RGB_HDR)
 			{
-				TextureFormatName = AndroidTexFormat::NameASTC_RGB_LDR;
+				TextureFormatName = GetTargetPlatformSettings()->GetFallbackASTCHDR();
 			}
 		}
 	}
diff --git a/Engine/Source/Developer/IOS/IOSTargetPlatformControls/Private/IOSTargetPlatformControls.cpp b/Engine/Source/Developer/IOS/IOSTargetPlatformControls/Private/IOSTargetPlatformControls.cpp
--- a/Engine/Source/Developer/IOS/IOSTargetPlatformControls/Private/IOSTargetPlatformControls.cpp
+++ b/Engine/Source/Developer/IOS/IOSTargetPlatformControls/Private/IOSTargetPlatformControls.cpp
@@ -564,7 +564,7 @@ void FIOSTargetPlatformControls::GetTextureFormats( const UTexture* Texture, TAr
 		{
 			if ( TextureFormatName == NameASTC_RGB_HDR )
 			{
-				TextureFormatName = NameASTC_RGB_LDR;
+				TextureFormatName = GetTargetPlatformSettings()->GetFallbackASTCHDR();
 			}
 		}
 	}
diff --git a/Engine/Source/Developer/TargetPlatform/Private/TargetPlatformBase.cpp b/Engine/Source/Developer/TargetPlatform/Private/TargetPlatformBase.cpp
--- a/Engine/Source/Developer/TargetPlatform/Private/TargetPlatformBase.cpp
+++ b/Engine/Source/Developer/TargetPlatform/Private/TargetPlatformBase.cpp
@@ -124,6 +124,22 @@ bool FTargetPlatformBase::UsesASTCHDR() const
 	return (bUsesARMCompressor && GASTCHDRProfile != 0);
 }
 
+FName FTargetPlatformBase::GetFallbackASTCHDR() const
+{
+	static IConsoleVariable* CVarFallbackForASTCHDRProfile = IConsoleManager::Get().FindConsoleVariable(TEXT("cook.FallbackForASTCHDRProfile"));
+	static const FName NameRGBA16F(TEXT("RGBA16F"));
+	static const FName NameASTC_RGB_LDR(TEXT("ASTC_RGB"));
+
+	switch (CVarFallbackForASTCHDRProfile->GetInt())
+	{
+	case 0:
+		return NameASTC_RGB_LDR;
+	case 1:
+		return NameRGBA16F;
+	}
+	return NameRGBA16F;
+}
+
 void FTargetPlatformBase::GetRayTracingShaderFormats(TArray<FName>& OutFormats) const
 {
 	if (UsesRayTracing())
diff --git a/Engine/Source/Developer/TargetPlatform/Private/TargetPlatformSettingsBase.cpp b/Engine/Source/Developer/TargetPlatform/Private/TargetPlatformSettingsBase.cpp
--- a/Engine/Source/Developer/TargetPlatform/Private/TargetPlatformSettingsBase.cpp
+++ b/Engine/Source/Developer/TargetPlatform/Private/TargetPlatformSettingsBase.cpp
@@ -4,6 +4,16 @@
 #include "AnalyticsEventAttribute.h"
 #include "HAL/IConsoleManager.h"
 
+int32 GFallbackASTCHDRProfile = 0;
+static FAutoConsoleVariableRef CVarFallbackASTCHDRProfile(
+	TEXT("cook.FallbackForASTCHDRProfile"),
+	GFallbackASTCHDRProfile,
+	TEXT("What Fallback Format we use if ASTC HDR Profile is not supported.")
+	TEXT("0 = Fallback to ASTC LDR")
+	TEXT("1 = Fallback to RGBA16F"),
+	ECVF_Default | ECVF_ReadOnly
+);
+
 bool FTargetPlatformSettingsBase::UsesForwardShading() const
 {
 	static IConsoleVariable* CVarForwardShading = IConsoleManager::Get().FindConsoleVariable(TEXT("r.ForwardShading"));
@@ -105,6 +115,23 @@ bool FTargetPlatformSettingsBase::UsesASTCHDR() const
 	return (bUsesARMCompressor && bUsesASTCHDR);
 }
 
+
+
+FName FTargetPlatformSettingsBase::GetFallbackASTCHDR() const
+{
+	static const FName NameRGBA16F(TEXT("RGBA16F"));
+	static const FName NameASTC_RGB_LDR(TEXT("ASTC_RGB"));
+
+	switch (GFallbackASTCHDRProfile)
+	{
+	case 0:
+		return NameASTC_RGB_LDR;
+	case 1:
+		return NameRGBA16F;
+	}
+	return NameRGBA16F;
+}
+
 void FTargetPlatformSettingsBase::GetRayTracingShaderFormats(TArray<FName>& OutFormats) const
 {
 	if (UsesRayTracing())
diff --git a/Engine/Source/Developer/TargetPlatform/Public/Common/TargetPlatformBase.h b/Engine/Source/Developer/TargetPlatform/Public/Common/TargetPlatformBase.h
--- a/Engine/Source/Developer/TargetPlatform/Public/Common/TargetPlatformBase.h
+++ b/Engine/Source/Developer/TargetPlatform/Public/Common/TargetPlatformBase.h
@@ -109,6 +109,8 @@ public:
 	TARGETPLATFORM_API virtual bool UsesMobileDBuffer() const override;
 
 	TARGETPLATFORM_API virtual bool UsesASTCHDR() const override;
+	
+	TARGETPLATFORM_API virtual FName GetFallbackASTCHDR() const override;
 
 	TARGETPLATFORM_API virtual void GetRayTracingShaderFormats(TArray<FName>& OutFormats) const override;
 
@@ -855,6 +857,10 @@ public:
 	{
 		return TargetPlatformSettings->UsesASTCHDR();
 	}
+	virtual FName GetFallbackASTCHDR() const override
+	{
+		return TargetPlatformSettings->GetFallbackASTCHDR();
+	}
 	virtual void GetAllPossibleShaderFormats(TArray<FName>& OutFormats) const override 
 	{
 		TargetPlatformSettings->GetAllPossibleShaderFormats(OutFormats);
diff --git a/Engine/Source/Developer/TargetPlatform/Public/Common/TargetPlatformSettingsBase.h b/Engine/Source/Developer/TargetPlatform/Public/Common/TargetPlatformSettingsBase.h
--- a/Engine/Source/Developer/TargetPlatform/Public/Common/TargetPlatformSettingsBase.h
+++ b/Engine/Source/Developer/TargetPlatform/Public/Common/TargetPlatformSettingsBase.h
@@ -37,6 +37,7 @@ public:
 	TARGETPLATFORM_API virtual bool UsesMobileAmbientOcclusion() const override;
 	TARGETPLATFORM_API virtual bool UsesMobileDBuffer() const override;
 	TARGETPLATFORM_API virtual bool UsesASTCHDR() const override;
+	TARGETPLATFORM_API virtual FName GetFallbackASTCHDR() const override;
 	TARGETPLATFORM_API virtual void GetRayTracingShaderFormats(TArray<FName>& OutFormats) const override;
 
 #if WITH_ENGINE
diff --git a/Engine/Source/Developer/TargetPlatform/Public/Interfaces/ITargetPlatformSettings.h b/Engine/Source/Developer/TargetPlatform/Public/Interfaces/ITargetPlatformSettings.h
--- a/Engine/Source/Developer/TargetPlatform/Public/Interfaces/ITargetPlatformSettings.h
+++ b/Engine/Source/Developer/TargetPlatform/Public/Interfaces/ITargetPlatformSettings.h
@@ -231,6 +231,11 @@ public:
 	 */
 	virtual bool UsesASTCHDR() const = 0;
 
+	/**
+	 * Gets the fallback format when ASTC HDR is not supported
+	 */
+	virtual FName GetFallbackASTCHDR() const = 0;
+
 	/**
 	* Gets the shader formats this platform can use.
 	*
