[Threading] Make initialization of TargetPlatformManagerModule functions thread-safe.

Static containers such as FormatVersionCache are initialized in ShaderFormatVersion() and friends but have no threading safe guards.
These functions can be used in a multi-threaded environment, such as FShaderCompileJob::GetInputHash() which can be called in both "ShaderCompilingThread" and "ShaderCompilingThread-Distributed".
This change uses RW locks to read the values with a leightweight read-only lock first and if not initialized yet it uses a second write lock.

#rnx

diff --git a/Engine/Source/Developer/TargetPlatform/Private/TargetPlatformManagerModule.cpp b/Engine/Source/Developer/TargetPlatform/Private/TargetPlatformManagerModule.cpp
--- a/Engine/Source/Developer/TargetPlatform/Private/TargetPlatformManagerModule.cpp
+++ b/Engine/Source/Developer/TargetPlatform/Private/TargetPlatformManagerModule.cpp
@@ -383,87 +383,25 @@ public:
 	virtual const TArray<ITargetPlatform*>& GetActiveTargetPlatforms() override
 	{
 		static bool bInitialized = false;
+		static FRWLock InitializedLock;
 		static TArray<ITargetPlatform*> Results;
 
-		if (!bInitialized || bForceCacheUpdate)
+		// Check if result is already initialized with lightweight read-only lock
 		{
-			bInitialized = true;
-			bHasInitErrors = false; // If we had errors before, reset the flag and see later in this function if there are errors.
-			InitErrorMessages.Empty();
-
-			Results.Empty(Results.Num());
-
-			const TArray<ITargetPlatform*>& TargetPlatforms = GetTargetPlatforms();	
-
-			FString PlatformStr;
-
-			if (FParse::Value(FCommandLine::Get(), TEXT("TARGETPLATFORM="), PlatformStr))
+			FRWScopeLock InitializedLockGuard(InitializedLock, FRWScopeLockType::SLT_ReadOnly);
+			if (bInitialized && !bForceCacheUpdate)
 			{
-				if (PlatformStr == TEXT("None"))
-				{
-				}
-				else if (PlatformStr == TEXT("All"))
-				{
-					Results = TargetPlatforms;
-				}
-				else
-				{
-					TArray<FString> PlatformNames;
-
-					PlatformStr.ParseIntoArray(PlatformNames, TEXT("+"), true);
-
-					for (int32 Index = 0; Index < TargetPlatforms.Num(); Index++)
-					{
-						if (PlatformNames.Contains(TargetPlatforms[Index]->PlatformName()))
-						{
-							Results.Add(TargetPlatforms[Index]);
-						}
-					}
-
-					if (Results.Num() == 0)
-					{
-						// An invalid platform was specified...
-						// Inform the user.
-						TStringBuilder<1024> AvailablePlatforms;
-						for (int32 Index = 0; Index < TargetPlatforms.Num(); Index++)
-						{
-							if (Index > 0)
-							{
-								AvailablePlatforms << TEXT(", ");
-							}
-							AvailablePlatforms << TargetPlatforms[Index]->PlatformName();
-						}
-						bHasInitErrors = true;
-						InitErrorMessages.Appendf(TEXT("Invalid target platform specified (%s). Available = { %s } "), *PlatformStr, *AvailablePlatforms);
-						UE_LOG(LogTargetPlatformManager, Error, TEXT("Invalid target platform specified (%s). Available = { %s } "), *PlatformStr, *AvailablePlatforms);
-						return Results;
-					}
-				}
-			}
-			else
-			{
-				// if there is no argument, use the current platform and only build formats that are actually needed to run.
-				bRestrictFormatsToRuntimeOnly = true;
-
-				for (int32 Index = 0; Index < TargetPlatforms.Num(); Index++)
-				{
-					if (TargetPlatforms[Index]->IsRunningPlatform())
-					{
-						Results.Add(TargetPlatforms[Index]);
-					}
-				}
+				return Results;
 			}
+		}
 
-			if (!Results.Num())
-			{
-				UE_LOG(LogTargetPlatformManager, Display, TEXT("Not building assets for any platform."));
-			}
-			else
+		// Initialize the result now
+		{
+			FRWScopeLock InitializedLockGuard(InitializedLock, FRWScopeLockType::SLT_Write);
+			// Check if another thread initialized the result before this thread received the write lock
+			if (!bInitialized || bForceCacheUpdate)
 			{
-				for (int32 Index = 0; Index < Results.Num(); Index++)
-				{
-					UE_LOG(LogTargetPlatformManager, Display, TEXT("Building Assets For %s"), *Results[Index]->PlatformName());
-				}
+				bInitialized = InitializeActiveTargetPlatforms(Results);
 			}
 		}
 
@@ -480,26 +418,41 @@ public:
 	virtual ITargetPlatform* GetRunningTargetPlatform() override
 	{
 		static bool bInitialized = false;
+		static FRWLock InitializedLock;
 		static ITargetPlatform* Result = nullptr;
 
-		if (!bInitialized || bForceCacheUpdate)
+		// Check if result is already initialized with lightweight read-only lock
 		{
-			bInitialized = true;
-			Result = nullptr;
-
-			const TArray<ITargetPlatform*>& TargetPlatforms = GetTargetPlatforms();	
+			FRWScopeLock InitializedLockGuard(InitializedLock, FRWScopeLockType::SLT_ReadOnly);
+			if (bInitialized && !bForceCacheUpdate)
+			{
+				return Result;
+			}
+		}
 
-			for (int32 Index = 0; Index < TargetPlatforms.Num(); Index++)
+		// Initialize the result now
+		{
+			FRWScopeLock InitializedLockGuard(InitializedLock, FRWScopeLockType::SLT_Write);
+			// Check if another thread initialized the result before this thread received the write lock
+			if (!bInitialized || bForceCacheUpdate)
 			{
-				if (TargetPlatforms[Index]->IsRunningPlatform())
+				Result = nullptr;
+
+				const TArray<ITargetPlatform*>& TargetPlatforms = GetTargetPlatforms();
+
+				for (int32 Index = 0; Index < TargetPlatforms.Num(); Index++)
 				{
-					 // we should not have two running platforms
-					checkf((Result == nullptr),
-						TEXT("Found multiple running platforms.\n\t%s\nand\n\t%s"),
-						*Result->PlatformName(),
-						*TargetPlatforms[Index]->PlatformName()
+					if (TargetPlatforms[Index]->IsRunningPlatform())
+					{
+						// we should not have two running platforms
+						checkf((Result == nullptr),
+							TEXT("Found multiple running platforms.\n\t%s\nand\n\t%s"),
+							*Result->PlatformName(),
+							*TargetPlatforms[Index]->PlatformName()
 						);
-					Result = TargetPlatforms[Index];
+						Result = TargetPlatforms[Index];
+						bInitialized = (Result != nullptr);
+					}
 				}
 			}
 		}
@@ -511,93 +464,26 @@ public:
 	const TArray<const FormatType*>& GetFormatsWithHints()
 	{
 		static bool bInitialized = false;
+		static FRWLock InitializedLock;
 		static TArray<const FormatType*> Results;
 
-		if (!bInitialized || bForceCacheUpdate)
+		// Check if result is already initialized with lightweight read-only lock
 		{
-			bInitialized = true;
-			Results.Empty(Results.Num());
-
-			// the functions for dealing with hints are only defined with Engine				
-#if WITH_ENGINE 
-			TArray<FName> SupportedFormatsByHints;
-			TSet<FName> RequiredFormats;
-			// gather the hinted formats, and the needed formats for all the active targetplatforms
-			TArray<ITargetPlatform*> TargetPlatforms = GetTargetPlatforms();
-			for (ITargetPlatform* Platform : TargetPlatforms)
-			{
-				TArray<FName> FormatHints;
-				HelperType::GetHintedModules(Platform, FormatHints);
-				for (FName HintedModuleName : FormatHints)
-				{
-					FormatModuleType* Module = FModuleManager::Get().LoadModulePtr<FormatModuleType>(HintedModuleName);
-					if (Module != nullptr)
-					{
-						FormatType* Format = HelperType::GetFormatFromModule(Module);
-						if (Format && !Results.Contains(Format))
-						{
-							// remember the module
-							Results.Add(Format);
-							// remember its formats
-							Format->GetSupportedFormats(SupportedFormatsByHints);
-						}
-					}
-				}
-
-				// remember the formats the TP needs
-				TArray<FName> PlatformRequiredFormats;
-				HelperType::GetRequiredFormats(Platform, PlatformRequiredFormats);
-				RequiredFormats.Append(PlatformRequiredFormats);
-			}
-
-			// make sure every required format was found above
-			bool bFoundAllFormats = true;
-			for (FName Format : RequiredFormats)
-			{
-				if (!SupportedFormatsByHints.Contains(Format))
-				{
-					UE_LOG(LogTargetPlatformManager, Log, TEXT("Unable to find %s format %s from hinted modules, loading all potential format modules to find it"), HelperType::GetFormatDesc(), *Format.ToString());
-					bFoundAllFormats = false;
-					break;
-				}
-			}
-
-			// if we found all the formats from the hints, we are done, and Results is filled out
-			if (bFoundAllFormats)
+			FRWScopeLock InitializedLockGuard(InitializedLock, FRWScopeLockType::SLT_ReadOnly);
+			if (bInitialized && !bForceCacheUpdate)
 			{
 				return Results;
 			}
-#endif
-
-			// if the hints weren't enough to find everything, then load all modules		
-			TArray<FName> Modules;
-			FModuleManager::Get().FindModules(HelperType::GetAllModuleWildcard(), Modules);
-
-			if (!Modules.Num())
-			{
-				UE_LOG(LogTargetPlatformManager, Error, TEXT("No target %s formats found!"), HelperType::GetFormatDesc());
-			}
+		}
 
-			for (int32 Index = 0; Index < Modules.Num(); Index++)
+		// Initialize the result now
+		{
+			FRWScopeLock InitializedLockGuard(InitializedLock, FRWScopeLockType::SLT_Write);
+			// Check if another thread initialized the result before this thread received the write lock
+			if (!bInitialized || bForceCacheUpdate)
 			{
-				FormatModuleType* Module = FModuleManager::LoadModulePtr<FormatModuleType>(Modules[Index]);
-				if (Module)
-				{
-
-					FormatType* Format = HelperType::GetFormatFromModule(Module);
-					UE_LOG(LogTargetPlatformManager, Log, TEXT("Loaded format module %s"), *Modules[Index].ToString());
-					if (Format != nullptr)
-					{
-						TArray<FName> Formats;
-						Format->GetSupportedFormats(Formats);
-						for (FName Name : Formats)
-						{
-							UE_LOG(LogTargetPlatformManager, Log, TEXT("  %s"), *Name.ToString());
-						}
-
-						Results.AddUnique(Format);
-					}
-				}
+				InitializeFormatsWithHints<FormatType, FormatModuleType, HelperType>(Results);
+				bInitialized = true;
 			}
 		}
 
@@ -708,37 +594,34 @@ public:
 	virtual uint32 ShaderFormatVersion(FName Name) override
 	{
 		static bool bInitialized = false;
+		static FRWLock InitializedLock;
 		static TMap<FName, uint32> FormatVersionCache;
 
-		if (!bInitialized || bForceCacheUpdate)
+		// Check if result is already initialized with lightweight read-only lock
 		{
-			FormatVersionCache.Reset();
-
-			for (const IShaderFormat* SF : GetShaderFormats())
+			FRWScopeLock InitializedLockGuard(InitializedLock, FRWScopeLockType::SLT_ReadOnly);
+			if (bInitialized && !bForceCacheUpdate)
 			{
-				TArray<FName> Formats;
-				SF->GetSupportedFormats(Formats);
-				for (FName FormatName : Formats)
-				{
-					FormatVersionCache.FindOrAdd(FormatName, SF->GetVersion(FormatName));
-				}
+				return FindShaderFormatVersion(FormatVersionCache, Name);
 			}
-			bInitialized = true;
 		}
 
-		const uint32* Result = FormatVersionCache.Find(Name);
-		if (!Result)
+		// Initialize the result now
 		{
-			UE_LOG(LogTargetPlatformManager, Fatal, TEXT("ShaderFormat not found for %s!  Dynamically loaded shader formats require invalidation of FormatVersionCache."), *Name.ToString());
-			return INDEX_NONE;
+			FRWScopeLock InitializedLockGuard(InitializedLock, FRWScopeLockType::SLT_Write);
+			// Check if another thread initialized the result before this thread received the write lock
+			if (!bInitialized || bForceCacheUpdate)
+			{
+				InitializeShaderFormatVersions(FormatVersionCache);
+				bInitialized = true;
+			}
 		}
 
-		return *Result;
+		return FindShaderFormatVersion(FormatVersionCache, Name);
 	}
 
 	virtual const TArray<const IPhysXCooking*>& GetPhysXCooking() override
 	{
-		static bool bInitialized = false;
 		static TArray<const IPhysXCooking*> Results;
 
 		return Results;
@@ -1362,7 +1245,203 @@ private:
 		}
 		bIgnoreFirstDelegateCall = false;
 	}
-	
+
+	bool InitializeActiveTargetPlatforms(TArray<ITargetPlatform*>& OutResults)
+	{
+		bHasInitErrors = false; // If we had errors before, reset the flag and see later in this function if there are errors.
+		InitErrorMessages.Empty();
+
+		OutResults.Empty(OutResults.Num());
+
+		const TArray<ITargetPlatform*>& TargetPlatforms = GetTargetPlatforms();
+
+		FString PlatformStr;
+
+		if (FParse::Value(FCommandLine::Get(), TEXT("TARGETPLATFORM="), PlatformStr))
+		{
+			if (PlatformStr == TEXT("None"))
+			{
+			}
+			else if (PlatformStr == TEXT("All"))
+			{
+				OutResults = TargetPlatforms;
+			}
+			else
+			{
+				TArray<FString> PlatformNames;
+
+				PlatformStr.ParseIntoArray(PlatformNames, TEXT("+"), true);
+
+				for (int32 Index = 0; Index < TargetPlatforms.Num(); Index++)
+				{
+					if (PlatformNames.Contains(TargetPlatforms[Index]->PlatformName()))
+					{
+						OutResults.Add(TargetPlatforms[Index]);
+					}
+				}
+
+				if (OutResults.Num() == 0)
+				{
+					// An invalid platform was specified...
+					// Inform the user.
+					TStringBuilder<1024> AvailablePlatforms;
+					for (int32 Index = 0; Index < TargetPlatforms.Num(); Index++)
+					{
+						if (Index > 0)
+						{
+							AvailablePlatforms << TEXT(", ");
+						}
+						AvailablePlatforms << TargetPlatforms[Index]->PlatformName();
+					}
+					bHasInitErrors = true;
+					InitErrorMessages.Appendf(TEXT("Invalid target platform specified (%s). Available = { %s } "), *PlatformStr, *AvailablePlatforms);
+					UE_LOG(LogTargetPlatformManager, Error, TEXT("Invalid target platform specified (%s). Available = { %s } "), *PlatformStr, *AvailablePlatforms);
+					return false;
+				}
+			}
+		}
+		else
+		{
+			// if there is no argument, use the current platform and only build formats that are actually needed to run.
+			bRestrictFormatsToRuntimeOnly = true;
+
+			for (int32 Index = 0; Index < TargetPlatforms.Num(); Index++)
+			{
+				if (TargetPlatforms[Index]->IsRunningPlatform())
+				{
+					OutResults.Add(TargetPlatforms[Index]);
+				}
+			}
+		}
+
+		if (!OutResults.Num())
+		{
+			UE_LOG(LogTargetPlatformManager, Display, TEXT("Not building assets for any platform."));
+		}
+		else
+		{
+			for (int32 Index = 0; Index < OutResults.Num(); Index++)
+			{
+				UE_LOG(LogTargetPlatformManager, Display, TEXT("Building Assets For %s"), *OutResults[Index]->PlatformName());
+			}
+		}
+
+		return true;
+	}
+
+	template<typename FormatType, typename FormatModuleType, typename HelperType>
+	void InitializeFormatsWithHints(TArray<const FormatType*>& OutResults)
+	{
+		OutResults.Empty(OutResults.Num());
+
+		// the functions for dealing with hints are only defined with Engine				
+#if WITH_ENGINE 
+		TArray<FName> SupportedFormatsByHints;
+		TSet<FName> RequiredFormats;
+		// gather the hinted formats, and the needed formats for all the active targetplatforms
+		TArray<ITargetPlatform*> TargetPlatforms = GetTargetPlatforms();
+		for (ITargetPlatform* Platform : TargetPlatforms)
+		{
+			TArray<FName> FormatHints;
+			HelperType::GetHintedModules(Platform, FormatHints);
+			for (FName HintedModuleName : FormatHints)
+			{
+				FormatModuleType* Module = FModuleManager::Get().LoadModulePtr<FormatModuleType>(HintedModuleName);
+				if (Module != nullptr)
+				{
+					FormatType* Format = HelperType::GetFormatFromModule(Module);
+					if (Format && !OutResults.Contains(Format))
+					{
+						// remember the module
+						OutResults.Add(Format);
+						// remember its formats
+						Format->GetSupportedFormats(SupportedFormatsByHints);
+					}
+				}
+			}
+
+			// remember the formats the TP needs
+			TArray<FName> PlatformRequiredFormats;
+			HelperType::GetRequiredFormats(Platform, PlatformRequiredFormats);
+			RequiredFormats.Append(PlatformRequiredFormats);
+		}
+
+		// make sure every required format was found above
+		bool bFoundAllFormats = true;
+		for (FName Format : RequiredFormats)
+		{
+			if (!SupportedFormatsByHints.Contains(Format))
+			{
+				UE_LOG(LogTargetPlatformManager, Log, TEXT("Unable to find %s format %s from hinted modules, loading all potential format modules to find it"), HelperType::GetFormatDesc(), *Format.ToString());
+				bFoundAllFormats = false;
+				break;
+			}
+		}
+
+		// if we found all the formats from the hints, we are done, and Results is filled out
+		if (bFoundAllFormats)
+		{
+			return;
+		}
+#endif
+
+		// if the hints weren't enough to find everything, then load all modules		
+		TArray<FName> Modules;
+		FModuleManager::Get().FindModules(HelperType::GetAllModuleWildcard(), Modules);
+
+		if (!Modules.Num())
+		{
+			UE_LOG(LogTargetPlatformManager, Error, TEXT("No target %s formats found!"), HelperType::GetFormatDesc());
+		}
+
+		for (int32 Index = 0; Index < Modules.Num(); Index++)
+		{
+			FormatModuleType* Module = FModuleManager::LoadModulePtr<FormatModuleType>(Modules[Index]);
+			if (Module)
+			{
+				FormatType* Format = HelperType::GetFormatFromModule(Module);
+				UE_LOG(LogTargetPlatformManager, Log, TEXT("Loaded format module %s"), *Modules[Index].ToString());
+				if (Format != nullptr)
+				{
+					TArray<FName> Formats;
+					Format->GetSupportedFormats(Formats);
+					for (FName Name : Formats)
+					{
+						UE_LOG(LogTargetPlatformManager, Log, TEXT("  %s"), *Name.ToString());
+					}
+
+					OutResults.AddUnique(Format);
+				}
+			}
+		}
+	}
+
+	void InitializeShaderFormatVersions(TMap<FName, uint32>& OutFormatVersionCache)
+	{
+		OutFormatVersionCache.Reset();
+
+		for (const IShaderFormat* SF : GetShaderFormats())
+		{
+			TArray<FName> Formats;
+			SF->GetSupportedFormats(Formats);
+			for (FName FormatName : Formats)
+			{
+				OutFormatVersionCache.FindOrAdd(FormatName, SF->GetVersion(FormatName));
+			}
+		}
+	}
+
+	uint32 FindShaderFormatVersion(const TMap<FName, uint32>& FormatVersionCache, FName Name) const
+	{
+		const uint32* Result = FormatVersionCache.Find(Name);
+		if (!Result)
+		{
+			UE_LOG(LogTargetPlatformManager, Fatal, TEXT("ShaderFormat not found for %s!  Dynamically loaded shader formats require invalidation of FormatVersionCache."), *Name.ToString());
+			return INDEX_NONE;
+		}
+		return *Result;
+	}
+
 	static FString SDKStatusMessage;
 	static void OnStatusOutput(FString Message)
 	{
