- Skip runnung UBT with -SetupPlatforms from the editor if it's determined to not be needed

diff --git a/Engine/Source/Developer/TargetPlatform/Private/TargetPlatformManagerModule.cpp b/Engine/Source/Developer/TargetPlatform/Private/TargetPlatformManagerModule.cpp
--- a/Engine/Source/Developer/TargetPlatform/Private/TargetPlatformManagerModule.cpp
+++ b/Engine/Source/Developer/TargetPlatform/Private/TargetPlatformManagerModule.cpp
@@ -30,6 +30,9 @@
 #include "DesktopPlatformModule.h"
 #include "Interfaces/ITurnkeySupportModule.h"
 #include "DataDrivenShaderPlatformInfo.h"
+#include "Serialization/JsonReader.h"
+#include "Serialization/JsonSerializer.h"
+#include "Misc/FileHelper.h"
 
 DEFINE_LOG_CATEGORY_STATIC(LogTargetPlatformManager, Log, All);
 
@@ -43,43 +46,302 @@ DEFINE_LOG_CATEGORY_STATIC(LogTargetPlatformManager, Log, All);
 #if AUTOSDKS_ENABLED
 namespace UE::AutoSDK
 {
-static bool IsAutoSDKsEnabled()
+	static const TCHAR* SDKRootEnvVar = TEXT("UE_SDKS_ROOT");
+	static const TCHAR* SDKInstallManifestFileName = TEXT("CurrentlyInstalled.txt");
+	static const TCHAR* SDKLastScriptRunVersionFileName = TEXT("CurrentlyInstalled.Version.txt");
+	static const TCHAR* SDKRequiredScriptVersionFileName = TEXT("Version.txt");
+	static const TCHAR* SDKEnvironmentVarsFile = TEXT("OutputEnvVars.txt");
+	
+	static FProcHandle AutoSDKSetupUBTProc;
+	static FConfigFile CachedInfoToSaveAfterUBT;
+
+
+	static bool IsAutoSDKsEnabled()
+	{
+		FString SDKPath = FPlatformMisc::GetEnvironmentVariable(SDKRootEnvVar);
+
+		// AutoSDKs only enabled if UE_SDKS_ROOT is set.
+		if (SDKPath.Len() != 0)
+		{
+			return true;
+		}
+		return false;
+	}
+
+	static FString GetProjectPathForUBT()
+	{
+		if (FPaths::IsProjectFilePathSet())
+		{
+			return FPaths::ConvertRelativePathToFull(FPaths::GetProjectFilePath());
+		}
+		if (FApp::HasProjectName())
+		{
+			FString ProjectPath = FPaths::ProjectDir() / FApp::GetProjectName() + TEXT(".uproject");
+			if (FPaths::FileExists(ProjectPath))
+			{
+				return ProjectPath;
+			}
+		}
+		return FString();
+	}
+
+	static FString GetAutoSDKPlatformDir(const TCHAR* Platform)
+	{
+		FString SDKPath = FPlatformMisc::GetEnvironmentVariable(SDKRootEnvVar);
+#if PLATFORM_WINDOWS
+		static FString HostPlatform(TEXT("HostWin64"));
+#else
+		static FString HostPlatform = FString::Printf(TEXT("Host%hs"), FPlatformProperties::IniPlatformName());
+#endif		
+
+		return FPaths::Combine(*SDKPath, *HostPlatform, Platform);
+	}
+
+// get a value from an SDK.json file
+static FString GetSDKInfo(const TCHAR* Platform, const TCHAR* Key)
 {
-	static const FString SDKRootEnvFar(TEXT("UE_SDKS_ROOT"));
+	// 1. start with per-project json
+	//   a. look in platform extension location, then non-platext location, for a file
+	//   b. if key not found, look up the ParentSdkFile chain (which is relative paths)
+	// 2. if key still not found, go back to 1. but using engine-wide SDK.json file
 
-	FString SDKPath = FPlatformMisc::GetEnvironmentVariable(*SDKRootEnvFar);
+	auto MakeConfigFilename = [Platform](const FString& RootDir) -> FString
+		{
+			FString PlatformExtensionLocation = FString::Printf(TEXT("%s/Platforms/%s/Config/%s_SDK.json"), *RootDir, Platform, Platform);
+			if (FPaths::FileExists(PlatformExtensionLocation))
+			{
+				return PlatformExtensionLocation;
+			}
+			FString StandardLocation = FString::Printf(TEXT("%s/Config/%s/%s_SDK.json"), *RootDir, Platform, Platform);
+			if (FPaths::FileExists(StandardLocation))
+			{
+				return StandardLocation;
+			}
+			return TEXT("");
+		};
+
+	auto GetJsonObject = [](const FString& Filename) -> TSharedPtr<FJsonObject>
+		{
+			FString FileContents;
+			if (!FFileHelper::LoadFileToString(FileContents, *Filename))
+			{
+				return nullptr;
+			}
 
-	// AutoSDKs only enabled if UE_SDKS_ROOT is set.
-	if (SDKPath.Len() != 0)
+			TSharedPtr<FJsonObject> Object;
+			TSharedRef<TJsonReader<> > Reader = TJsonReaderFactory<>::Create(FileContents);
+			if (!FJsonSerializer::Deserialize(Reader, Object) || !Object.IsValid())
+			{
+				return nullptr;
+			}
+
+			return Object;
+		};
+
+	// Steps 1 or 2 above
+	for (int Pass = 0; Pass < 2; Pass++)
 	{
-		return true;
+		// Step 1 - ProjectDir; Step 2 - EngineDir
+		FString RootDir = Pass == 0 ? GetProjectPathForUBT() : FPaths::EngineDir();
+		// Step a. Choose platformextension or standard location
+		FString InfoFile = RootDir.IsEmpty() ? FString() : MakeConfigFilename(RootDir);
+
+		// make sure we found a file
+		if (InfoFile.IsEmpty())
+		{
+			continue;
+		}
+
+		FString ParentFile;
+		do
+		{
+			// read the file and get the key if it exists
+			TSharedPtr<FJsonObject> JsonObj = GetJsonObject(InfoFile);
+			if (JsonObj != nullptr)
+			{
+				if (JsonObj->HasField(Key))
+				{
+					// if we have a value, we are done!
+					return JsonObj->GetStringField(Key);
+				}
+				// step b. go up parent chain
+				if (JsonObj->HasField(TEXT("ParentSDKFile")))
+				{
+					ParentFile = JsonObj->GetStringField(TEXT("ParentSDKFile"));
+
+					// the parent field is relative to current json file
+					InfoFile = FPaths::GetPath(InfoFile) / ParentFile;
+					FPaths::NormalizeFilename(InfoFile);
+				}
+				else
+				{
+					ParentFile = TEXT("");
+				}
+			}
+		} while (!ParentFile.IsEmpty());
 	}
-	return false;
+
+	// if we got no value, then return empty string for failure
+	return FString();
 }
 
-static FString GetProjectPathForUBT()
+// Check to see if any platforms are not up to date with AutoSDK setup. This is how it's determined:
+// 1. Load an existing <proj>/Intermediate/CachedAutoSdkInfo.ini to get cached info from last run
+// 2. -> Compare current vs cached value for UE_SDKS_ROOT
+// 3. For each (real) platform:
+//   a. Get the current AutoSDKVersion from <platform>_SDK.json (looking in project overrides, and parent SDK.json chain as needed)
+//   b. -> Compare current vs cached AutoSDKVersion
+//   c. Set TargetSDKRoot to UE_SDKS_ROOT/<Platform>/<AutoSDKVersion>
+//   d. Check if TargetSDKRoot/AutoSDKVersion/setup.bat|sh [aka SetupScript] exists
+//   e. -> Compare current vs cached SetupScript existence
+//   f. Read TargetSDKRoot/CurrentlyInstalled.txt
+//   g. Get installed version, and if it's AutoSDK or ManualSDK
+//   h. If ManualSDK, skip and move to next platform
+//   i. -> Compare installed version to AutoSDKVersion
+//   j. -> Compare installed "script version" to required "script version" (this manages updating setup.bat)
+//   j. -> Check for existence of TargetSDKRoot/OutputEnvVars.txt
+// 4. If anything above (marked with ->) doesn't match up, then run UBT which will update CurrentlyInstalled.txt and OutputEnvVars.txt
+// 5. [much later] Read in OutputEnvVars.txt (either from step 4, or a previous UBT run) to set envvars for the editor process
+//
+// Still todo: 
+//   Manage CurrentlyInstalled.Version.txt, which allows for forcing a re-run of the SetupScript, which would mean that we need to run UBT
+//   Double check what happens if a user goes from manual->auto or auto->manual
+//   We can't tell here if the platform prefers auto vs manual, which may affect the previous line
+//   ManualSDKEnvVars.txt? I believe this is not used and we can ignore it
+//   Maybe force run UBT if UBT.dll has been updated since last run of editor, just in case
+static bool NeedsToRunSetupPlatformsUBT()
 {
-	if (FPaths::IsProjectFilePathSet())
+	// early out and always return true on builders
+	if (GIsBuildMachine)
 	{
-		return FPaths::ConvertRelativePathToFull(FPaths::GetProjectFilePath());
+		return true;
+	}
+
+	SCOPED_BOOT_TIMING("AutoSDKInit - NeedsToRunSetupPlatformsUBT");
+
+	FString CachedAutoSDKPath = FPaths::ProjectIntermediateDir() / TEXT("CachedAutoSdkInfo.ini");
+	FConfigFile CachedInfo;
+	CachedInfo.Read(CachedAutoSDKPath);
+
+	// first check that the UE_SDKS_ROOT var is the same
+	FString SDKsRoot = FPlatformMisc::GetEnvironmentVariable(UE::AutoSDK::SDKRootEnvVar);
+	FString CachedValue;
+	CachedInfo.GetString(TEXT("Global"), TEXT("SDKsRoot"), CachedValue);
+	CachedInfoToSaveAfterUBT.SetString(TEXT("Global"), TEXT("SDKsRoot"), *SDKsRoot);
+
+	// even if we need to run, we finish this function to get all the cached info to check next run
+	bool bNeedsToRun = false;
+	if (SDKsRoot != CachedValue)
+	{
+		UE_LOG(LogTargetPlatformManager, Log, TEXT("Running UBT for AutoSDK init because the value of %s changed"), SDKRootEnvVar);
+		bNeedsToRun = true;
 	}
-	if (FApp::HasProjectName())
+
+	for (const TPair<FName, FDataDrivenPlatformInfo>& Pair : FDataDrivenPlatformInfoRegistry::GetAllPlatformInfos())
 	{
-		FString ProjectPath = FPaths::ProjectDir() / FApp::GetProjectName() + TEXT(".uproject");
-		if (FPaths::FileExists(ProjectPath))
+		// check only real platforms, that have an AutoSDK path
+		if (!Pair.Value.bEnabledForUse || Pair.Value.bIsFakePlatform || Pair.Value.AutoSDKPath.IsEmpty())
 		{
-			return ProjectPath;
+			continue;
+		}
+
+		// get the AutoSDK version - AutoSDKDirectory if specified, or use default of MainVersion
+		FString PlatformName = Pair.Key.ToString();
+		FString AutoSDKVersion = GetSDKInfo(*PlatformName, TEXT("AutoSDKDirectory"));
+		if (AutoSDKVersion.IsEmpty())
+		{
+			AutoSDKVersion = GetSDKInfo(*PlatformName, TEXT("MainVersion"));
+		}
+		CachedInfo.GetString(*PlatformName, TEXT("AutoSDKDirectory"), CachedValue);
+		CachedInfoToSaveAfterUBT.SetString(*PlatformName, TEXT("AutoSDKDirectory"), *AutoSDKVersion);
+
+		if (AutoSDKVersion != CachedValue)
+		{
+			UE_LOG(LogTargetPlatformManager, Log, TEXT("Running UBT for AutoSDK init because AutoSDK version for %s changed"), *PlatformName);
+			bNeedsToRun = true;
+		}
+
+		FString TargetSDKRoot = GetAutoSDKPlatformDir(*PlatformName);
+
+		// check if presence of setup.bat, etc matches
+#if PLATFORM_WINDOWS
+		FString SetupScript = TargetSDKRoot / AutoSDKVersion / "setup.bat";
+#else
+		FString SetupScript = TargetSDKRoot / AutoSDKVersion / "setup.sh";
+#endif
+		bool bSetupScriptExists = FPaths::FileExists(SetupScript);
+		bool bCachedValue;
+		CachedInfo.GetBool(*PlatformName, TEXT("bSetupScriptExists"), bCachedValue);
+		CachedInfoToSaveAfterUBT.SetBool(*PlatformName, TEXT("bSetupScriptExists"), bSetupScriptExists);
+		if (bSetupScriptExists != bCachedValue)
+		{
+			UE_LOG(LogTargetPlatformManager, Log, TEXT("Running UBT for AutoSDK init because the existence of %s changed"), *SetupScript);
+			bNeedsToRun = true;
+		}
+
+		// only deal with autosdk stuff if the setupscript actually exists now - if it doesn't, none of this will do anything in UBT
+		if (bSetupScriptExists)
+		{
+			// check the currently installed manifest file - if the version doesn't match and it's not a ManualSDK installation,
+			// we will need to run UBT
+			TArray<FString> FileLines;
+			bool bIsManualInstall = false;
+			if (FFileHelper::LoadFileToStringArray(FileLines, *(TargetSDKRoot / SDKInstallManifestFileName)) &&
+				FileLines.Num() > 1)
+			{
+				if (FileLines[1] == TEXT("ManualSDK"))
+				{
+					bIsManualInstall = true;
+				}
+				if (FileLines[0] != AutoSDKVersion && !bIsManualInstall)
+				{
+					UE_LOG(LogTargetPlatformManager, Log, TEXT("Running UBT for AutoSDK init because the installed version of %s doesn't match the required version"), *PlatformName);
+					bNeedsToRun = true;
+				}
+			}
+			else
+			{
+				UE_LOG(LogTargetPlatformManager, Log, TEXT("Running UBT for AutoSDK init because the the file %s is invalid"), *(TargetSDKRoot / SDKInstallManifestFileName));
+				bNeedsToRun = true;
+			}
+
+			if (!bIsManualInstall)
+			{
+				// check if the last run script version matches what the AutoSDK version requires
+				// if they mismatch, or if the installed version is missing, run UBT
+				FString InstalledScriptVersion;
+				FString RequiredScriptVersion;
+				if (!FFileHelper::LoadFileToString(RequiredScriptVersion, *(TargetSDKRoot / AutoSDKVersion / SDKRequiredScriptVersionFileName)))
+				{
+					RequiredScriptVersion = TEXT("UnspecifiedScriptVersion");
+				}
+				if (!FFileHelper::LoadFileToString(InstalledScriptVersion, *(TargetSDKRoot / SDKLastScriptRunVersionFileName)) ||
+					RequiredScriptVersion.TrimStartAndEnd() != InstalledScriptVersion.TrimStartAndEnd())
+				{
+					UE_LOG(LogTargetPlatformManager, Log, TEXT("Running UBT for AutoSDK init because the last run script version in %s doesn't match required (%s !- %s)"), *(TargetSDKRoot / SDKLastScriptRunVersionFileName), *InstalledScriptVersion, *RequiredScriptVersion);
+					bNeedsToRun = true;
+				}
+			}
+
+			// autosdk writes out the OUtputEnvvars.txt file that should exists if autosdk ran
+			if (!bIsManualInstall && !FPaths::FileExists(*(TargetSDKRoot / SDKEnvironmentVarsFile)))
+			{
+				UE_LOG(LogTargetPlatformManager, Log, TEXT("Running UBT for AutoSDK init because the file %s is missing"), *(TargetSDKRoot / SDKEnvironmentVarsFile));
+				bNeedsToRun = true;
+			}
 		}
 	}
-	return FString();
+
+	return bNeedsToRun;
 }
 
 // kick off a call to UBT nice and early so that it's results are hopefully ready when needed
-static FProcHandle AutoSDKSetupUBTProc;
 FDelayedAutoRegisterHelper GAutoSDKInit(EDelayedRegisterRunPhase::FileSystemReady, []
 	{
+		SCOPED_BOOT_TIMING("AutoSDKInit - DelayedAutoRegisterHelper");
+
 		// amortize UBT cost by calling it once for all platforms, rather than once per platform.
-		if (IsAutoSDKsEnabled() && FParse::Param(FCommandLine::Get(), TEXT("Multiprocess")) == false)
+		if (IsAutoSDKsEnabled() && FParse::Param(FCommandLine::Get(), TEXT("Multiprocess")) == false && NeedsToRunSetupPlatformsUBT())
 		{
 			FString UBTParams(TEXT("-Mode=SetupPlatforms"));
 			FString Project = GetProjectPathForUBT();
@@ -208,13 +470,24 @@ public:
 		{					
 			if (UE::AutoSDK::AutoSDKSetupUBTProc.IsValid())
 			{
-				SCOPED_BOOT_TIMING("FTargetPlatformManagerModule - WaitForUBTProc");
+				SCOPED_BOOT_TIMING("AutoSDKInit - WaitForUBTProc");
+				// wait for UBT to finish
 				FPlatformProcess::WaitForProc(UE::AutoSDK::AutoSDKSetupUBTProc);
+
+				// if it failed, we don't want to cache the saved info, so we make sure we run it again next time
+				int32 ReturnCode = 0;
+				FPlatformProcess::GetProcReturnCode(UE::AutoSDK::AutoSDKSetupUBTProc, &ReturnCode);
+
+				if (ReturnCode == 0)
+				{
+					// now save out the cached info so next run we can skip UBT
+					UE::AutoSDK::CachedInfoToSaveAfterUBT.Write(*(FPaths::ProjectIntermediateDir() / TEXT("CachedAutoSdkInfo.ini")));
+				}
 			}
 
 			// we have to setup our local environment according to AutoSDKs or the ITargetPlatform's IsSDkInstalled calls may fail
 			// before we get a change to setup for a given platform.  Use the platforminfo list to avoid any kind of interdependency.
-			SCOPED_BOOT_TIMING("FTargetPlatformManagerModule.SetupAndValidateAutoSDK");
+			SCOPED_BOOT_TIMING("AutoSDKInit - SetupAndValidateAutoSDK");
 			for (auto Pair: FDataDrivenPlatformInfoRegistry::GetAllPlatformInfos())
 			{
 				if (Pair.Value.AutoSDKPath.Len() > 0)
@@ -887,32 +1160,14 @@ protected:
 		FString HostPlatform = FString::Printf(TEXT("Host%hs"), FPlatformProperties::IniPlatformName());
 #endif		
 
-		static const FString SDKRootEnvFar(TEXT("UE_SDKS_ROOT"));
-		FString SDKPath = FPlatformMisc::GetEnvironmentVariable(*SDKRootEnvFar);
+		FString TargetSDKRoot = UE::AutoSDK::GetAutoSDKPlatformDir(*AutoSDKPath);
+		FString SDKInstallManifestFilePath = TargetSDKRoot / UE::AutoSDK::SDKInstallManifestFileName;
 
-		FString TargetSDKRoot = FPaths::Combine(*SDKPath, *HostPlatform, *AutoSDKPath);
-		static const FString SDKInstallManifestFileName(TEXT("CurrentlyInstalled.txt"));
-		FString SDKInstallManifestFilePath = FPaths::Combine(*TargetSDKRoot, *SDKInstallManifestFileName);
-
-		// If we are using a manual install, then it is valid for there to be no OutputEnvVars file.
-		TUniquePtr<FArchive> InstallManifestFile(IFileManager::Get().CreateFileReader(*SDKInstallManifestFilePath));
-		if (InstallManifestFile)
+		TArray<FString> FileLines;
+		// If we are using a manual install, then it is valid for there to be no OutputEnvVars file, so 
+		// check for ManualINstall in the CurrentlyInstalled file
+		if (FFileHelper::LoadFileToStringArray(FileLines, *SDKInstallManifestFilePath))
 		{
-			TArray<FString> FileLines;
-			int64 FileSize = InstallManifestFile->TotalSize();
-			int64 MemSize = FileSize + 1;
-			void* FileMem = FMemory::Malloc(MemSize);
-			FMemory::Memset(FileMem, 0, MemSize);
-
-			InstallManifestFile->Serialize(FileMem, FileSize);
-
-			FString FileAsString(ANSI_TO_TCHAR(FileMem));
-			FileAsString.ParseIntoArrayLines(FileLines);
-
-			FMemory::Free(FileMem);
-			InstallManifestFile->Close();
-
-			
 			if (FileLines.Num() != 2 && FileLines.Num() != 3)
 			{
 				UE_LOG(LogTargetPlatformManager, Warning, TEXT("Malformed install manifest file for Platform %s"), *AutoSDKPath);
@@ -1489,4 +1744,4 @@ private:
 FString FTargetPlatformManagerModule::SDKStatusMessage = TEXT("");
 
 
-IMPLEMENT_MODULE(FTargetPlatformManagerModule, TargetPlatform);
+IMPLEMENT_MODULE(FTargetPlatformManagerModule, TargetPlatform);
\ No newline at end of file
diff --git a/Engine/Source/Developer/TargetPlatform/TargetPlatform.Build.cs b/Engine/Source/Developer/TargetPlatform/TargetPlatform.Build.cs
--- a/Engine/Source/Developer/TargetPlatform/TargetPlatform.Build.cs
+++ b/Engine/Source/Developer/TargetPlatform/TargetPlatform.Build.cs
@@ -12,6 +12,8 @@ public class TargetPlatform : ModuleRules
 		PrivateDependencyModuleNames.Add("Slate");
 		PrivateDependencyModuleNames.Add("Projects");
 		PrivateDependencyModuleNames.Add("RenderCore");
+		PrivateDependencyModuleNames.Add("Json");
+
 		PublicDependencyModuleNames.Add("DeveloperSettings");
 		PublicDependencyModuleNames.Add("AudioPlatformConfiguration");
 		PublicDependencyModuleNames.Add("DesktopPlatform");
