Fix crash on shutdown due to wrong static destruction order between CachedInfoToSaveAfterUBT and TSOnFConfigDeleted delegate

diff --git a/Engine/Source/Developer/TargetPlatform/Private/TargetPlatformManagerModule.cpp b/Engine/Source/Developer/TargetPlatform/Private/TargetPlatformManagerModule.cpp
--- a/Engine/Source/Developer/TargetPlatform/Private/TargetPlatformManagerModule.cpp
+++ b/Engine/Source/Developer/TargetPlatform/Private/TargetPlatformManagerModule.cpp
@@ -53,8 +53,12 @@ namespace UE::AutoSDK
 	static const TCHAR* SDKEnvironmentVarsFile = TEXT("OutputEnvVars.txt");
 	
 	static FProcHandle AutoSDKSetupUBTProc;
-	static FConfigFile CachedInfoToSaveAfterUBT;
 
+	static FConfigFile& GetCachedInfoToSaveAfterUBT()
+	{
+		static FConfigFile CachedInfoToSaveAfterUBT;
+		return CachedInfoToSaveAfterUBT;
+	}
 
 	static bool IsAutoSDKsEnabled()
 	{
@@ -227,7 +231,7 @@ static bool NeedsToRunSetupPlatformsUBT()
 	FString SDKsRoot = FPlatformMisc::GetEnvironmentVariable(UE::AutoSDK::SDKRootEnvVar);
 	FString CachedValue;
 	CachedInfo.GetString(TEXT("Global"), TEXT("SDKsRoot"), CachedValue);
-	CachedInfoToSaveAfterUBT.SetString(TEXT("Global"), TEXT("SDKsRoot"), *SDKsRoot);
+	GetCachedInfoToSaveAfterUBT().SetString(TEXT("Global"), TEXT("SDKsRoot"), *SDKsRoot);
 
 	// even if we need to run, we finish this function to get all the cached info to check next run
 	bool bNeedsToRun = false;
@@ -253,7 +257,7 @@ static bool NeedsToRunSetupPlatformsUBT()
 			AutoSDKVersion = GetSDKInfo(*PlatformName, TEXT("MainVersion"));
 		}
 		CachedInfo.GetString(*PlatformName, TEXT("AutoSDKDirectory"), CachedValue);
-		CachedInfoToSaveAfterUBT.SetString(*PlatformName, TEXT("AutoSDKDirectory"), *AutoSDKVersion);
+		GetCachedInfoToSaveAfterUBT().SetString(*PlatformName, TEXT("AutoSDKDirectory"), *AutoSDKVersion);
 
 		if (AutoSDKVersion != CachedValue)
 		{
@@ -272,7 +276,7 @@ static bool NeedsToRunSetupPlatformsUBT()
 		bool bSetupScriptExists = FPaths::FileExists(SetupScript);
 		bool bCachedValue;
 		CachedInfo.GetBool(*PlatformName, TEXT("bSetupScriptExists"), bCachedValue);
-		CachedInfoToSaveAfterUBT.SetBool(*PlatformName, TEXT("bSetupScriptExists"), bSetupScriptExists);
+		GetCachedInfoToSaveAfterUBT().SetBool(*PlatformName, TEXT("bSetupScriptExists"), bSetupScriptExists);
 		if (bSetupScriptExists != bCachedValue)
 		{
 			UE_LOG(LogTargetPlatformManager, Log, TEXT("Running UBT for AutoSDK init because the existence of %s changed"), *SetupScript);
@@ -481,7 +485,7 @@ public:
 				if (ReturnCode == 0)
 				{
 					// now save out the cached info so next run we can skip UBT
-					UE::AutoSDK::CachedInfoToSaveAfterUBT.Write(*(FPaths::ProjectIntermediateDir() / TEXT("CachedAutoSdkInfo.ini")));
+					UE::AutoSDK::GetCachedInfoToSaveAfterUBT().Write(*(FPaths::ProjectIntermediateDir() / TEXT("CachedAutoSdkInfo.ini")));
 				}
 			}
 
