Fix wrong path being used for AutoSDK platform dir in TargetPlatformManagerModule.cpp

diff --git a/Engine/Source/Developer/TargetPlatform/Private/TargetPlatformManagerModule.cpp b/Engine/Source/Developer/TargetPlatform/Private/TargetPlatformManagerModule.cpp
--- a/Engine/Source/Developer/TargetPlatform/Private/TargetPlatformManagerModule.cpp
+++ b/Engine/Source/Developer/TargetPlatform/Private/TargetPlatformManagerModule.cpp
@@ -89,18 +89,6 @@ namespace UE::AutoSDK
 		return FString();
 	}
 
-	static FString GetAutoSDKPlatformDir(const TCHAR* Platform)
-	{
-		FString SDKPath = FPlatformMisc::GetEnvironmentVariable(SDKRootEnvVar);
-#if PLATFORM_WINDOWS
-		static FString HostPlatform(TEXT("HostWin64"));
-#else
-		static FString HostPlatform = FString::Printf(TEXT("Host%hs"), FPlatformProperties::IniPlatformName());
-#endif		
-
-		return FPaths::Combine(*SDKPath, *HostPlatform, Platform);
-	}
-
 // get a value from an SDK.json file
 static FString GetSDKInfo(const TCHAR* Platform, const TCHAR* Key)
 {
@@ -189,6 +177,20 @@ static FString GetSDKInfo(const TCHAR* Platform, const TCHAR* Key)
 	return FString();
 }
 
+	static FString GetAutoSDKPlatformDir(const TCHAR* Platform)
+	{
+		FString SDKPath = FPlatformMisc::GetEnvironmentVariable(SDKRootEnvVar);
+#if PLATFORM_WINDOWS
+		static FString HostPlatform(TEXT("HostWin64"));
+#else
+		static FString HostPlatform = FString::Printf(TEXT("Host%hs"), FPlatformProperties::IniPlatformName());
+#endif		
+		FString AutoSDKPlatform = GetSDKInfo(Platform, TEXT("AutoSDKPlatform"));
+
+		return FPaths::Combine(*SDKPath, *HostPlatform, AutoSDKPlatform.Len() > 0 ? *AutoSDKPlatform : Platform);
+	}
+
+
 // Check to see if any platforms are not up to date with AutoSDK setup. This is how it's determined:
 // 1. Load an existing <proj>/Intermediate/CachedAutoSdkInfo.ini to get cached info from last run
 // 2. -> Compare current vs cached value for UE_SDKS_ROOT
