[HWRT] Expose Skeletal Mesh Offline BVH mode in platform settings.

diff --git a/Engine/Source/Developer/TargetPlatform/Private/TargetPlatformBase.cpp b/Engine/Source/Developer/TargetPlatform/Private/TargetPlatformBase.cpp
--- a/Engine/Source/Developer/TargetPlatform/Private/TargetPlatformBase.cpp
+++ b/Engine/Source/Developer/TargetPlatform/Private/TargetPlatformBase.cpp
@@ -84,6 +84,11 @@ bool FTargetPlatformBase::GetStaticMeshOfflineBVHCompression() const
 	return false;
 }
 
+EOfflineBVHMode FTargetPlatformBase::GetSkeletalMeshOfflineBVHMode() const
+{
+	return EOfflineBVHMode::Disabled;
+}
+
 bool FTargetPlatformBase::ForcesSimpleSkyDiffuse() const
 {
 	return false;
diff --git a/Engine/Source/Developer/TargetPlatform/Private/TargetPlatformSettingsBase.cpp b/Engine/Source/Developer/TargetPlatform/Private/TargetPlatformSettingsBase.cpp
--- a/Engine/Source/Developer/TargetPlatform/Private/TargetPlatformSettingsBase.cpp
+++ b/Engine/Source/Developer/TargetPlatform/Private/TargetPlatformSettingsBase.cpp
@@ -80,6 +80,11 @@ bool FTargetPlatformSettingsBase::GetStaticMeshOfflineBVHCompression() const
 	return false;
 }
 
+EOfflineBVHMode FTargetPlatformSettingsBase::GetSkeletalMeshOfflineBVHMode() const
+{
+	return EOfflineBVHMode::Disabled;
+}
+
 bool FTargetPlatformSettingsBase::ForcesSimpleSkyDiffuse() const
 {
 	return false;
diff --git a/Engine/Source/Developer/TargetPlatform/Public/Common/TargetPlatformBase.h b/Engine/Source/Developer/TargetPlatform/Public/Common/TargetPlatformBase.h
--- a/Engine/Source/Developer/TargetPlatform/Public/Common/TargetPlatformBase.h
+++ b/Engine/Source/Developer/TargetPlatform/Public/Common/TargetPlatformBase.h
@@ -100,6 +100,8 @@ public:
 
 	TARGETPLATFORM_API virtual bool GetStaticMeshOfflineBVHCompression() const override;
 
+	TARGETPLATFORM_API virtual EOfflineBVHMode GetSkeletalMeshOfflineBVHMode() const override;
+
 	TARGETPLATFORM_API virtual bool ForcesSimpleSkyDiffuse() const override;
 
 	TARGETPLATFORM_API virtual float GetDownSampleMeshDistanceFieldDivider() const override;
@@ -839,6 +841,10 @@ public:
 	{
 		return TargetPlatformSettings->GetStaticMeshOfflineBVHCompression();
 	}
+	virtual EOfflineBVHMode GetSkeletalMeshOfflineBVHMode() const override
+	{
+		return TargetPlatformSettings->GetSkeletalMeshOfflineBVHMode();
+	}
 	virtual bool ForcesSimpleSkyDiffuse() const override 
 	{
 		return TargetPlatformSettings->ForcesSimpleSkyDiffuse();
diff --git a/Engine/Source/Developer/TargetPlatform/Public/Common/TargetPlatformSettingsBase.h b/Engine/Source/Developer/TargetPlatform/Public/Common/TargetPlatformSettingsBase.h
--- a/Engine/Source/Developer/TargetPlatform/Public/Common/TargetPlatformSettingsBase.h
+++ b/Engine/Source/Developer/TargetPlatform/Public/Common/TargetPlatformSettingsBase.h
@@ -32,6 +32,7 @@ public:
 	TARGETPLATFORM_API virtual uint32 GetSupportedHardwareMask() const override;
 	TARGETPLATFORM_API virtual EOfflineBVHMode GetStaticMeshOfflineBVHMode() const override;
 	TARGETPLATFORM_API virtual bool GetStaticMeshOfflineBVHCompression() const override;
+	TARGETPLATFORM_API virtual EOfflineBVHMode GetSkeletalMeshOfflineBVHMode() const override;
 	TARGETPLATFORM_API virtual bool ForcesSimpleSkyDiffuse() const override;
 	TARGETPLATFORM_API virtual float GetDownSampleMeshDistanceFieldDivider() const override;
 	TARGETPLATFORM_API virtual int32 GetHeightFogModeForOpaque() const override;
diff --git a/Engine/Source/Developer/TargetPlatform/Public/Interfaces/ITargetPlatformSettings.h b/Engine/Source/Developer/TargetPlatform/Public/Interfaces/ITargetPlatformSettings.h
--- a/Engine/Source/Developer/TargetPlatform/Public/Interfaces/ITargetPlatformSettings.h
+++ b/Engine/Source/Developer/TargetPlatform/Public/Interfaces/ITargetPlatformSettings.h
@@ -198,6 +198,11 @@ public:
 	*/
 	virtual bool GetStaticMeshOfflineBVHCompression() const = 0;
 
+	/**
+	* Gets skeletal mesh offline BVH mode
+	*/
+	virtual EOfflineBVHMode GetSkeletalMeshOfflineBVHMode() const = 0;
+
 	/**
 	* Gets whether the platform will use SH2 instead of SH3 for sky irradiance.
 	*/
