TargetPlatform: Add more logging to help diagnosing the cook error "LogTargetPlatformManager: Error: Invalid target platform specified (PlatformXXX). Available = { PlatformYYY, PlatformZZZ, ...}" that can be caused by many different underlying issues.
Always print the path to CachedAutoSdkInfo.ini to the log to increase awareness of this feature and make it easy to force a run of the UBT command.
Always print the command line parameters for UBT to the log to make it easy to run the same command on the side.
Run UBT with "-Mode=ValidatePlatforms -AllPlatforms -OutputSDKs -log=AutoSDKInfo.txt -verbose -timestamps" instead of "-Mode=SetupPlatforms" to get more relevant output.
Whenever an existing TargetPlatformModule fails to load, retry once with error logging enabled.
In ModuleManager Always print out-of-date modules and manifests to the log to increase awareness of this feature.
Make ValidatePlatformsMode.cs more robust by catching and printing exceptions thrown by GetMainVersion().

#rnx

diff --git a/Engine/Source/Developer/TargetPlatform/Private/TargetPlatformManagerModule.cpp b/Engine/Source/Developer/TargetPlatform/Private/TargetPlatformManagerModule.cpp
--- a/Engine/Source/Developer/TargetPlatform/Private/TargetPlatformManagerModule.cpp
+++ b/Engine/Source/Developer/TargetPlatform/Private/TargetPlatformManagerModule.cpp
@@ -26,6 +26,7 @@
 #include "Interfaces/ITextureFormat.h"
 #include "Interfaces/ITextureFormatManagerModule.h"
 #include "Interfaces/ITextureFormatModule.h"
+#include "Logging/StructuredLog.h"
 #include "PlatformInfo.h"
 #include "DesktopPlatformModule.h"
 #include "Interfaces/ITurnkeySupportModule.h"
@@ -338,6 +339,7 @@ static bool NeedsToRunSetupPlatformsUBT()
 		}
 	}
 
+	UE_LOGFMT(LogTargetPlatformManager, Log, "{CachedAutoSDKPath} is {Status}", CachedAutoSDKPath, bNeedsToRun ? "out-of-date" : "up-to-date");
 	return bNeedsToRun;
 }
 
@@ -347,24 +349,37 @@ FDelayedAutoRegisterHelper GAutoSDKInit(EDelayedRegisterRunPhase::FileSystemRead
 		SCOPED_BOOT_TIMING("AutoSDKInit - DelayedAutoRegisterHelper");
 
 		// amortize UBT cost by calling it once for all platforms, rather than once per platform.
-		if (IsAutoSDKsEnabled() && FParse::Param(FCommandLine::Get(), TEXT("Multiprocess")) == false && NeedsToRunSetupPlatformsUBT())
+		if (IsAutoSDKsEnabled())
 		{
-			FString UBTParams(TEXT("-Mode=SetupPlatforms"));
-			FString Project = GetProjectPathForUBT();
-			if (Project.Len() > 0)
+			FString UBTParams(TEXT("-Mode=ValidatePlatforms -AllPlatforms -OutputSDKs"));
 			{
-				UBTParams += FString::Printf(TEXT(" -project=%s"), *Project);
+				if (FString Project = GetProjectPathForUBT(); Project.Len() > 0)
+				{
+					UBTParams += FString::Printf(TEXT(" -project=%s"), *Project);
+				}
+				// Write the output to a separate file since it performs much better than read pipes in this scenario
+				// where the invocation script is quite involved and is in turn calling other scripts and programs...
+				{
+					FString LogFile = FPaths::Combine(FPaths::ProjectLogDir(), "AutoSDKInfo.txt");
+					FString AbsLogFile = IFileManager::Get().ConvertToAbsolutePathForExternalAppForWrite(*LogFile);
+					UBTParams += FString::Printf(TEXT(" -log=\"%s\" -verbose -timestamps"), *AbsLogFile);
+				}
 			}
-				
-			int32 UBTReturnCode = -1;
-			FString UBTOutput;
+			if (FParse::Param(FCommandLine::Get(), TEXT("Multiprocess")) == false && NeedsToRunSetupPlatformsUBT())
+			{
+				int32 UBTReturnCode = -1;
 
-			void* ReadPipe = nullptr;
-			void* WritePipe = nullptr;
-			AutoSDKSetupUBTProc = FDesktopPlatformModule::Get()->InvokeUnrealBuildToolAsync(UBTParams, *GLog, ReadPipe, WritePipe, true);
-			if (!AutoSDKSetupUBTProc.IsValid())
+				void* ReadPipe = nullptr;
+				void* WritePipe = nullptr;
+				AutoSDKSetupUBTProc = FDesktopPlatformModule::Get()->InvokeUnrealBuildToolAsync(UBTParams, *GLog, ReadPipe, WritePipe, true);
+				if (!AutoSDKSetupUBTProc.IsValid())
+				{
+					UE_LOG(LogTargetPlatformManager, Warning, TEXT("AutoSDK is enabled (UE_SDKS_ROOT is set), but failed to run UBT to check SDK status! Check your installation."));
+				}
+			}
+			else
 			{
-				UE_LOG(LogTargetPlatformManager, Warning, TEXT("AutoSDK is enabled (UE_SDKS_ROOT is set), but failed to run UBT to check SDK status! Check your installation."));
+				UE_LOGFMT(LogTargetPlatformManager, Log, "Skip running UBT AutoSDK init with params [{UBTParams}]", UBTParams);
 			}
 		}
 	}
@@ -483,6 +498,9 @@ public:
 				// if it failed, we don't want to cache the saved info, so we make sure we run it again next time
 				int32 ReturnCode = 0;
 				FPlatformProcess::GetProcReturnCode(UE::AutoSDK::AutoSDKSetupUBTProc, &ReturnCode);
+				FPlatformProcess::CloseProc(UE::AutoSDK::AutoSDKSetupUBTProc);
+				UE::AutoSDK::AutoSDKSetupUBTProc.Reset();
+				UE_LOGFMT(LogTargetPlatformManager, Log, "UBT AutoSDK ReturnCode: {ReturnCode}", ReturnCode);
 
 				if (ReturnCode == 0)
 				{
@@ -928,22 +946,24 @@ protected:
 		ITargetPlatformControlsModule* ModuleControls = nullptr;
 		ITargetPlatformSettingsModule* ModuleSettings = nullptr;
 
-		if (!FModuleManager::Get().ModuleExists(*PlatformModuleName.ToString()))
+		FModuleManager& ModuleManager = FModuleManager::Get();
+
+		if (!ModuleManager.ModuleExists(*PlatformModuleName.ToString()))
 		{
 			PlatformModuleName = *(PlatformName.ToString() + TEXT("TargetPlatform"));
 		}
 
-		if (FModuleManager::Get().ModuleExists(*PlatformModuleName.ToString()))
+		if (ModuleManager.ModuleExists(*PlatformModuleName.ToString()))
 		{
 			Module = FModuleManager::LoadModulePtr<ITargetPlatformModule>(PlatformModuleName);
 		}
 
-		if (FModuleManager::Get().ModuleExists(*PlatformSettingsModuleName.ToString()))
+		if (ModuleManager.ModuleExists(*PlatformSettingsModuleName.ToString()))
 		{
 			ModuleSettings = FModuleManager::LoadModulePtr<ITargetPlatformSettingsModule>(PlatformSettingsModuleName);
 		}
 
-		if (FModuleManager::Get().ModuleExists(*PlatformControlsModuleName.ToString()))
+		if (ModuleManager.ModuleExists(*PlatformControlsModuleName.ToString()))
 		{
 			ModuleControls = FModuleManager::LoadModulePtr<ITargetPlatformControlsModule>(PlatformControlsModuleName);
 		}
@@ -961,6 +981,25 @@ protected:
 			}
 		}
 
+		if (Module == nullptr)
+		{
+			if (ModuleManager.ModuleExists(*PlatformModuleName.ToString()))
+			{
+				// Retry module load with error logging enabled
+				EModuleLoadResult Result;
+				Module = static_cast<ITargetPlatformModule*>(ModuleManager.LoadModuleWithFailureReason(PlatformModuleName, Result, ELoadModuleFlags::LogFailures));
+				UE_CLOGFMT(!Module, LogTargetPlatformManager, Warning,
+					"Failed to load module '{PlatformModuleName}' for platform {PlatformName} (Reason={Reason}, AutoSDKPath='{AutoSDKPath}').",
+					PlatformModuleName, PlatformName.ToString(), LexToString(Result), AutoSDKPath);
+			}
+			else
+			{
+				UE_LOGFMT(LogTargetPlatformManager, Log,
+					"Failed to load module '{PlatformModuleName}' for platform {PlatformName} (Reason={Reason}, AutoSDKPath='{AutoSDKPath}').",
+					PlatformModuleName, PlatformName.ToString(), LexToString(EModuleLoadResult::FileNotFound), AutoSDKPath);
+			}
+		}
+
 		// original logic for module loading here
 		if (Module)
 		{
@@ -1749,4 +1788,4 @@ private:
 FString FTargetPlatformManagerModule::SDKStatusMessage = TEXT("");
 
 
-IMPLEMENT_MODULE(FTargetPlatformManagerModule, TargetPlatform);
\ No newline at end of file
+IMPLEMENT_MODULE(FTargetPlatformManagerModule, TargetPlatform);
diff --git a/Engine/Source/Programs/UnrealBuildTool/Modes/ValidatePlatformsMode.cs b/Engine/Source/Programs/UnrealBuildTool/Modes/ValidatePlatformsMode.cs
--- a/Engine/Source/Programs/UnrealBuildTool/Modes/ValidatePlatformsMode.cs
+++ b/Engine/Source/Programs/UnrealBuildTool/Modes/ValidatePlatformsMode.cs
@@ -1,5 +1,6 @@
 // Copyright Epic Games, Inc. All Rights Reserved.
 
+using System;
 using System.Collections.Generic;
 using System.Threading.Tasks;
 using EpicGames.Core;
@@ -50,13 +51,25 @@ namespace UnrealBuildTool
 			}
 
 			// Output a line for each registered platform
+			List<string> ExceptionMessages = new List<string>();
 			foreach (UnrealTargetPlatform Platform in Platforms)
 			{
 				UEBuildPlatform.TryGetBuildPlatform(Platform, out UEBuildPlatform? BuildPlatform);
 				string PlatformSDKString = "";
 				if (bOutputSDKs)
 				{
-					PlatformSDKString = BuildPlatform != null ? UEBuildPlatform.GetSDK(Platform)!.GetMainVersion() : "<UNKNOWN>";
+					PlatformSDKString = "<UNKNOWN>";
+					if (BuildPlatform != null)
+					{
+						try
+						{
+							PlatformSDKString = UEBuildPlatform.GetSDK(Platform)!.GetMainVersion();
+						}
+						catch (Exception Ex)
+						{
+							ExceptionMessages.Add(Ex.Message.TrimEnd());
+						}
+					}
 				}
 
 				if (BuildPlatform != null && BuildPlatform.HasRequiredSDKsInstalled() == SDKStatus.Valid)
@@ -68,6 +81,10 @@ namespace UnrealBuildTool
 					Logger.LogInformation("##PlatformValidate: {Platform} INVALID {PlatformSdkString}", Platform.ToString(), PlatformSDKString);
 				}
 			}
+			foreach (string Message in ExceptionMessages)
+			{
+				Logger.LogInformation(Message);
+			}
 			return Task.FromResult(0);
 		}
 	}
diff --git a/Engine/Source/Runtime/Core/Private/Modules/ModuleManager.cpp b/Engine/Source/Runtime/Core/Private/Modules/ModuleManager.cpp
--- a/Engine/Source/Runtime/Core/Private/Modules/ModuleManager.cpp
+++ b/Engine/Source/Runtime/Core/Private/Modules/ModuleManager.cpp
@@ -4,6 +4,7 @@
 
 #include "HAL/FileManager.h"
 #include "Internationalization/StringTableCore.h"
+#include "Logging/StructuredLog.h"
 #include "Misc/App.h"
 #include "Misc/DataDrivenPlatformInfoRegistry.h"
 #include "Misc/DateTime.h"
@@ -733,7 +734,7 @@ IModuleInterface* FModuleManager::GetOrLoadModule(const FName InModuleName, EMod
 	return Result;
 }
 
-static const TCHAR* LexToString(EModuleLoadResult LoadResult)
+const TCHAR* LexToString(EModuleLoadResult LoadResult)
 {
 	switch (LoadResult)
 	{
@@ -1660,14 +1661,27 @@ void FModuleManager::FindModulePathsInDirectory(const FString& InDirectoryName,
 	SearchDirectoryNames.Insert(InDirectoryName, 0);
 
 	// Enumerate the modules in each directory
-	for(const FString& SearchDirectoryName: SearchDirectoryNames)
+	for (const FString& SearchDirectoryName : SearchDirectoryNames)
 	{
 		FModuleManifest Manifest;
-		if (FModuleManifest::TryRead(FModuleManifest::GetFileName(SearchDirectoryName, bIsGameDirectory), Manifest) && (Manifest.BuildId == BuildId.GetValue() || SearchDirectoryName.Contains(TEXT("/Engine/Plugins/Bridge/"))))
+		FString Filename = FModuleManifest::GetFileName(SearchDirectoryName, bIsGameDirectory);
+		if (FModuleManifest::TryRead(Filename, Manifest))
 		{
-			for (const TPair<FString, FString>& Pair : Manifest.ModuleNameToFileName)
+			if ((Manifest.BuildId == BuildId.GetValue() || SearchDirectoryName.Contains(TEXT("/Engine/Plugins/Bridge/"))))
 			{
-				OutModulePaths.Add(FName(*Pair.Key), *FPaths::Combine(*SearchDirectoryName, *Pair.Value));
+				for (const TPair<FString, FString>& Pair : Manifest.ModuleNameToFileName)
+				{
+					OutModulePaths.Add(FName(*Pair.Key), *FPaths::Combine(*SearchDirectoryName, *Pair.Value));
+				}
+			}
+			else
+			{
+				UE_LOGFMT(LogModuleManager, Log, "Skipping out-of-date modules in manifest '{Filename}' (BuildId {ModuleBuildId} != {BuildId}):",
+					Filename, Manifest.BuildId, BuildId.GetValue());
+				for (const TPair<FString, FString>& Pair : Manifest.ModuleNameToFileName)
+				{
+					UE_LOGFMT(LogModuleManager, Log, "    Skipping module '{Filename}'.", *FPaths::Combine(*SearchDirectoryName, *Pair.Value));
+				}
 			}
 		}
 	}
diff --git a/Engine/Source/Runtime/Core/Public/Modules/ModuleManager.h b/Engine/Source/Runtime/Core/Public/Modules/ModuleManager.h
--- a/Engine/Source/Runtime/Core/Public/Modules/ModuleManager.h
+++ b/Engine/Source/Runtime/Core/Public/Modules/ModuleManager.h
@@ -66,6 +66,8 @@ enum class EModuleLoadResult
 	NotLoadedByGameThread
 };
 
+CORE_API const TCHAR* LexToString(EModuleLoadResult LoadResult);
+
 /**
  * Enumerates reasons for failed module unloads.
  */
