- Fixed an issue with SDK.json file parenting when a file fails to parse

diff --git a/Engine/Source/Developer/TargetPlatform/Private/TargetPlatformManagerModule.cpp b/Engine/Source/Developer/TargetPlatform/Private/TargetPlatformManagerModule.cpp
--- a/Engine/Source/Developer/TargetPlatform/Private/TargetPlatformManagerModule.cpp
+++ b/Engine/Source/Developer/TargetPlatform/Private/TargetPlatformManagerModule.cpp
@@ -145,9 +145,11 @@ static FString GetSDKInfo(const TCHAR* Platform, const TCHAR* Key)
 			continue;
 		}
 
-		FString ParentFile;
+		bool bContinueUpToParent;
 		do
 		{
+			bContinueUpToParent = false;
+
 			// read the file and get the key if it exists
 			TSharedPtr<FJsonObject> JsonObj = GetJsonObject(InfoFile);
 			if (JsonObj != nullptr)
@@ -160,18 +162,18 @@ static FString GetSDKInfo(const TCHAR* Platform, const TCHAR* Key)
 				// step b. go up parent chain
 				if (JsonObj->HasField(TEXT("ParentSDKFile")))
 				{
-					ParentFile = JsonObj->GetStringField(TEXT("ParentSDKFile"));
-
 					// the parent field is relative to current json file
-					InfoFile = FPaths::GetPath(InfoFile) / ParentFile;
+					InfoFile = FPaths::GetPath(InfoFile) / JsonObj->GetStringField(TEXT("ParentSDKFile"));
 					FPaths::NormalizeFilename(InfoFile);
+
+					bContinueUpToParent = true;
 				}
-				else
-				{
-					ParentFile = TEXT("");
-				}
 			}
-		} while (!ParentFile.IsEmpty());
+			else
+			{
+				UE_LOG(LogTargetPlatformManager, Error, TEXT("Failed to parse SDK json file '%s'"), *InfoFile);
+			}
+		} while (bContinueUpToParent);
 	}
 
 	// if we got no value, then return empty string for failure
