Adding some defaulted state to this struct

#Synced-CL 48632724

diff --git a/Engine/Source/Developer/TargetPlatform/Public/Interfaces/IAudioFormat.h b/Engine/Source/Developer/TargetPlatform/Public/Interfaces/IAudioFormat.h
--- a/Engine/Source/Developer/TargetPlatform/Public/Interfaces/IAudioFormat.h
+++ b/Engine/Source/Developer/TargetPlatform/Public/Interfaces/IAudioFormat.h
@@ -15,22 +15,22 @@
 struct FSoundQualityInfo
 {
 	/** Holds the quality value ranging from 1 [poor] to 100 [very good]. */
-	int32 Quality;
+	int32 Quality = 0;
 
 	/** Holds the number of distinct audio channels. */
-	uint32 NumChannels;
+	uint32 NumChannels = 0;
 
 	/** Holds the number of PCM samples per second. */
-	uint32 SampleRate;
+	uint32 SampleRate = 0;
 
 	/** Holds the size of sample data in bytes. */
-	uint32 SampleDataSize;
+	uint32 SampleDataSize = 0;
 
 	/** Holds the length of the sound in seconds. */
-	float Duration;
+	float Duration = 0.f;
 
 	/** Holds whether the sound will be streamed. */
-	bool bStreaming;
+	bool bStreaming = false;
 
 	/** Holds a string for debugging purposes. */
 	FString DebugName;
