When AutoSDK is on and the UE_SDK_ROOT env variable is set
-Run UBT once, then cache the results
-If AutoSDKPath is set for the platform, then run SetupAutoSDK, if its not set then we have nothing to setup and assume its valid.

When AutoSDK is on and the UE_SDK_ROOT env variable is not set or AutoSDK is off
- Run UBT every time
- Check UBT output to gather valid platforms

When UBT not found, not available
- Only RunningPlatform TargetPlatformControl is created, all the other platforms only load the TargetPlatformSettings for preview.
- This is why we moved IsRunningPlatform from TPControl to TPSettings

#Synced-CL 49633005

diff --git a/Engine/Source/Developer/Android/AndroidPlatformEditor/Private/AndroidSDKSettings.cpp b/Engine/Source/Developer/Android/AndroidPlatformEditor/Private/AndroidSDKSettings.cpp
--- a/Engine/Source/Developer/Android/AndroidPlatformEditor/Private/AndroidSDKSettings.cpp
+++ b/Engine/Source/Developer/Android/AndroidPlatformEditor/Private/AndroidSDKSettings.cpp
@@ -64,7 +64,7 @@ void UAndroidSDKSettings::UpdateTargetModulePaths()
 	
 	if (Keys.Num() != 0)
 	{
-		TargetManagerModule->UpdatePlatformEnvironment(TEXT("Android"), Keys, Values);
+		TargetManagerModule->SetupEnvironmentVariables(Keys, Values);
 		AndroidDeviceDetection->UpdateADBPath();
 	}
 }
diff --git a/Engine/Source/Developer/CookedEditor/Public/CookedEditorTargetPlatformControls.h b/Engine/Source/Developer/CookedEditor/Public/CookedEditorTargetPlatformControls.h
--- a/Engine/Source/Developer/CookedEditor/Public/CookedEditorTargetPlatformControls.h
+++ b/Engine/Source/Developer/CookedEditor/Public/CookedEditorTargetPlatformControls.h
@@ -78,11 +78,6 @@ public:
 		}
 	}
 
-	virtual bool IsRunningPlatform() const override
-	{
-		return false;
-	}
-
 	virtual float GetVariantPriority() const override
 	{
 		// by returning -1, we will never use this variant when targeting host platform this class implements
@@ -198,11 +193,6 @@ public:
 		return false;
 	}
 
-	virtual bool IsRunningPlatform() const override
-	{
-		return false;
-	}
-
 	virtual float GetVariantPriority() const override
 	{
 		// by returning -1, we will never use this variant when targeting host platform this class implements
diff --git a/Engine/Source/Developer/Linux/LinuxTargetPlatformControls/Public/LinuxTargetPlatformControls.h b/Engine/Source/Developer/Linux/LinuxTargetPlatformControls/Public/LinuxTargetPlatformControls.h
--- a/Engine/Source/Developer/Linux/LinuxTargetPlatformControls/Public/LinuxTargetPlatformControls.h
+++ b/Engine/Source/Developer/Linux/LinuxTargetPlatformControls/Public/LinuxTargetPlatformControls.h
@@ -164,12 +164,6 @@ public:
 		return nullptr;
 	}
 
-	virtual bool IsRunningPlatform() const override
-	{
-		// Must be Linux platform as editor for this to be considered a running platform
-		return PLATFORM_LINUX && !UE_SERVER && !UE_GAME && WITH_EDITOR && TProperties::HasEditorOnlyData();
-	}
-
 	virtual bool IsSdkInstalled(bool bProjectHasCode, FString& OutDocumentationPath) const override
 	{
 		if (!PLATFORM_LINUX)
diff --git a/Engine/Source/Developer/Linux/LinuxTargetPlatformSettings/Public/LinuxTargetPlatformSettings.h b/Engine/Source/Developer/Linux/LinuxTargetPlatformSettings/Public/LinuxTargetPlatformSettings.h
--- a/Engine/Source/Developer/Linux/LinuxTargetPlatformSettings/Public/LinuxTargetPlatformSettings.h
+++ b/Engine/Source/Developer/Linux/LinuxTargetPlatformSettings/Public/LinuxTargetPlatformSettings.h
@@ -110,6 +110,12 @@ public:
 		}
 	}
 
+	virtual bool IsRunningPlatform() const override
+	{
+		// Must be Linux platform as editor for this to be considered a running platform
+		return PLATFORM_LINUX && !UE_SERVER && !UE_GAME && WITH_EDITOR && TProperties::HasEditorOnlyData();
+	}
+
 #if WITH_ENGINE
 	virtual void GetReflectionCaptureFormats(TArray<FName>& OutFormats) const override
 	{
diff --git a/Engine/Source/Developer/Mac/MacTargetPlatformControls/Public/GenericMacTargetPlatformControls.h b/Engine/Source/Developer/Mac/MacTargetPlatformControls/Public/GenericMacTargetPlatformControls.h
--- a/Engine/Source/Developer/Mac/MacTargetPlatformControls/Public/GenericMacTargetPlatformControls.h
+++ b/Engine/Source/Developer/Mac/MacTargetPlatformControls/Public/GenericMacTargetPlatformControls.h
@@ -81,12 +81,6 @@ public:
 		return NULL;
 	}
 
-	virtual bool IsRunningPlatform( ) const override
-	{
-		// Must be Mac platform as editor for this to be considered a running platform
-		return PLATFORM_MAC && !UE_SERVER && !UE_GAME && WITH_EDITOR && HAS_EDITOR_DATA;
-	}
-
 
 	virtual void GetPlatformSpecificProjectAnalytics( TArray<FAnalyticsEventAttribute>& AnalyticsParamArray ) const override
 	{
diff --git a/Engine/Source/Developer/Mac/MacTargetPlatformSettings/Public/GenericMacTargetPlatformSettings.h b/Engine/Source/Developer/Mac/MacTargetPlatformSettings/Public/GenericMacTargetPlatformSettings.h
--- a/Engine/Source/Developer/Mac/MacTargetPlatformSettings/Public/GenericMacTargetPlatformSettings.h
+++ b/Engine/Source/Developer/Mac/MacTargetPlatformSettings/Public/GenericMacTargetPlatformSettings.h
@@ -137,6 +137,12 @@ public:
 
 	//~ End ITargetPlatform Interface
 
+
+	virtual bool IsRunningPlatform() const override
+	{
+		// Must be Mac platform as editor for this to be considered a running platform
+		return PLATFORM_MAC && !UE_SERVER && !UE_GAME && WITH_EDITOR && HAS_EDITOR_DATA;
+	}
 private:
 
 #if WITH_ENGINE
diff --git a/Engine/Source/Developer/TargetPlatform/Private/TargetPlatformManagerModule.cpp b/Engine/Source/Developer/TargetPlatform/Private/TargetPlatformManagerModule.cpp
--- a/Engine/Source/Developer/TargetPlatform/Private/TargetPlatformManagerModule.cpp
+++ b/Engine/Source/Developer/TargetPlatform/Private/TargetPlatformManagerModule.cpp
@@ -42,9 +42,6 @@ DEFINE_LOG_CATEGORY_STATIC(LogTargetPlatformManager, Log, All);
 #define AUTOSDKS_ENABLED DDPI_HAS_EXTENDED_PLATFORMINFO_DATA
 #endif
 
-
-
-#if AUTOSDKS_ENABLED
 namespace UE::AutoSDK
 {
 	static const TCHAR* SDKRootEnvVar = TEXT("UE_SDKS_ROOT");
@@ -52,43 +49,46 @@ namespace UE::AutoSDK
 	static const TCHAR* SDKLastScriptRunVersionFileName = TEXT("CurrentlyInstalled.Version.txt");
 	static const TCHAR* SDKRequiredScriptVersionFileName = TEXT("Version.txt");
 	static const TCHAR* SDKEnvironmentVarsFile = TEXT("OutputEnvVars.txt");
-	
-	static FProcHandle AutoSDKSetupUBTProc;
+}
+
+static FProcHandle SDKSetupUBTProc;
+
+static FConfigFile& GetCachedInfoToSaveAfterUBT()
+{
+	static FConfigFile CachedInfoToSaveAfterUBT;
+	return CachedInfoToSaveAfterUBT;
+}
 
-	static FConfigFile& GetCachedInfoToSaveAfterUBT()
+static bool IsAutoSDKsEnabled()
+{
+#if AUTOSDKS_ENABLED
+	FString SDKPath = FPlatformMisc::GetEnvironmentVariable(UE::AutoSDK::SDKRootEnvVar);
+
+	// AutoSDKs only enabled if UE_SDKS_ROOT is set.
+	if (SDKPath.Len() != 0)
 	{
-		static FConfigFile CachedInfoToSaveAfterUBT;
-		return CachedInfoToSaveAfterUBT;
+		return true;
 	}
+#endif
+	return false;
+}
 
-	static bool IsAutoSDKsEnabled()
+static FString GetProjectPathForUBT()
+{
+	if (FPaths::IsProjectFilePathSet())
 	{
-		FString SDKPath = FPlatformMisc::GetEnvironmentVariable(SDKRootEnvVar);
-
-		// AutoSDKs only enabled if UE_SDKS_ROOT is set.
-		if (SDKPath.Len() != 0)
-		{
-			return true;
-		}
-		return false;
+		return FPaths::ConvertRelativePathToFull(FPaths::GetProjectFilePath());
 	}
-
-	static FString GetProjectPathForUBT()
+	if (FApp::HasProjectName())
 	{
-		if (FPaths::IsProjectFilePathSet())
+		FString ProjectPath = FPaths::ProjectDir() / FApp::GetProjectName() + TEXT(".uproject");
+		if (FPaths::FileExists(ProjectPath))
 		{
-			return FPaths::ConvertRelativePathToFull(FPaths::GetProjectFilePath());
+			return ProjectPath;
 		}
-		if (FApp::HasProjectName())
-		{
-			FString ProjectPath = FPaths::ProjectDir() / FApp::GetProjectName() + TEXT(".uproject");
-			if (FPaths::FileExists(ProjectPath))
-			{
-				return ProjectPath;
-			}
-		}
-		return FString();
 	}
+	return FString();
+}
 
 // get a value from an SDK.json file
 static FString GetSDKInfo(const TCHAR* Platform, const TCHAR* Key)
@@ -180,19 +180,18 @@ static FString GetSDKInfo(const TCHAR* Platform, const TCHAR* Key)
 	return FString();
 }
 
-	static FString GetAutoSDKPlatformDir(const TCHAR* Platform)
-	{
-		FString SDKPath = FPlatformMisc::GetEnvironmentVariable(SDKRootEnvVar);
+static FString GetAutoSDKPlatformDir(const TCHAR* Platform)
+{
+	FString SDKPath = FPlatformMisc::GetEnvironmentVariable(UE::AutoSDK::SDKRootEnvVar);
 #if PLATFORM_WINDOWS
-		static FString HostPlatform(TEXT("HostWin64"));
+	static FString HostPlatform(TEXT("HostWin64"));
 #else
-		static FString HostPlatform = FString::Printf(TEXT("Host%hs"), FPlatformProperties::IniPlatformName());
+	static FString HostPlatform = FString::Printf(TEXT("Host%hs"), FPlatformProperties::IniPlatformName());
 #endif		
-		FString AutoSDKPlatform = GetSDKInfo(Platform, TEXT("AutoSDKPlatform"));
-
-		return FPaths::Combine(*SDKPath, *HostPlatform, AutoSDKPlatform.Len() > 0 ? *AutoSDKPlatform : Platform);
-	}
+	FString AutoSDKPlatform = GetSDKInfo(Platform, TEXT("AutoSDKPlatform"));
 
+	return FPaths::Combine(*SDKPath, *HostPlatform, AutoSDKPlatform.Len() > 0 ? *AutoSDKPlatform : Platform);
+}
 
 // Check to see if any platforms are not up to date with AutoSDK setup. This is how it's determined:
 // 1. Load an existing <proj>/Intermediate/CachedAutoSdkInfo.ini to get cached info from last run
@@ -220,6 +219,11 @@ static FString GetSDKInfo(const TCHAR* Platform, const TCHAR* Key)
 //   Maybe force run UBT if UBT.dll has been updated since last run of editor, just in case
 static bool NeedsToRunSetupPlatformsUBT()
 {
+	if (!IsAutoSDKsEnabled())
+	{
+		return true;
+	}
+
 	// early out and always return true on builders
 	if (GIsBuildMachine)
 	{
@@ -242,7 +246,7 @@ static bool NeedsToRunSetupPlatformsUBT()
 	bool bNeedsToRun = false;
 	if (SDKsRoot != CachedValue)
 	{
-		UE_LOG(LogTargetPlatformManager, Log, TEXT("Running UBT for AutoSDK init because the value of %s changed"), SDKRootEnvVar);
+		UE_LOG(LogTargetPlatformManager, Log, TEXT("Running UBT for AutoSDK init because the value of %s changed"), UE::AutoSDK::SDKRootEnvVar);
 		bNeedsToRun = true;
 	}
 
@@ -295,7 +299,7 @@ static bool NeedsToRunSetupPlatformsUBT()
 			// we will need to run UBT
 			TArray<FString> FileLines;
 			bool bIsManualInstall = false;
-			if (FFileHelper::LoadFileToStringArray(FileLines, *(TargetSDKRoot / SDKInstallManifestFileName)) &&
+			if (FFileHelper::LoadFileToStringArray(FileLines, *(TargetSDKRoot / UE::AutoSDK::SDKInstallManifestFileName)) &&
 				FileLines.Num() > 1)
 			{
 				if (FileLines[1] == TEXT("ManualSDK"))
@@ -310,7 +314,7 @@ static bool NeedsToRunSetupPlatformsUBT()
 			}
 			else
 			{
-				UE_LOG(LogTargetPlatformManager, Log, TEXT("Running UBT for AutoSDK init because the the file %s is invalid"), *(TargetSDKRoot / SDKInstallManifestFileName));
+				UE_LOG(LogTargetPlatformManager, Log, TEXT("Running UBT for AutoSDK init because the the file %s is invalid"), *(TargetSDKRoot / UE::AutoSDK::SDKInstallManifestFileName));
 				bNeedsToRun = true;
 			}
 
@@ -320,22 +324,22 @@ static bool NeedsToRunSetupPlatformsUBT()
 				// if they mismatch, or if the installed version is missing, run UBT
 				FString InstalledScriptVersion;
 				FString RequiredScriptVersion;
-				if (!FFileHelper::LoadFileToString(RequiredScriptVersion, *(TargetSDKRoot / AutoSDKVersion / SDKRequiredScriptVersionFileName)))
+				if (!FFileHelper::LoadFileToString(RequiredScriptVersion, *(TargetSDKRoot / AutoSDKVersion / UE::AutoSDK::SDKRequiredScriptVersionFileName)))
 				{
 					RequiredScriptVersion = TEXT("UnspecifiedScriptVersion");
 				}
-				if (!FFileHelper::LoadFileToString(InstalledScriptVersion, *(TargetSDKRoot / SDKLastScriptRunVersionFileName)) ||
+				if (!FFileHelper::LoadFileToString(InstalledScriptVersion, *(TargetSDKRoot / UE::AutoSDK::SDKLastScriptRunVersionFileName)) ||
 					RequiredScriptVersion.TrimStartAndEnd() != InstalledScriptVersion.TrimStartAndEnd())
 				{
-					UE_LOG(LogTargetPlatformManager, Log, TEXT("Running UBT for AutoSDK init because the last run script version in %s doesn't match required (%s !- %s)"), *(TargetSDKRoot / SDKLastScriptRunVersionFileName), *InstalledScriptVersion, *RequiredScriptVersion);
+					UE_LOG(LogTargetPlatformManager, Log, TEXT("Running UBT for AutoSDK init because the last run script version in %s doesn't match required (%s !- %s)"), *(TargetSDKRoot / UE::AutoSDK::SDKLastScriptRunVersionFileName), *InstalledScriptVersion, *RequiredScriptVersion);
 					bNeedsToRun = true;
 				}
 			}
 
 			// autosdk writes out the OUtputEnvvars.txt file that should exists if autosdk ran
-			if (!bIsManualInstall && !FPaths::FileExists(*(TargetSDKRoot / SDKEnvironmentVarsFile)))
+			if (!bIsManualInstall && !FPaths::FileExists(*(TargetSDKRoot / UE::AutoSDK::SDKEnvironmentVarsFile)))
 			{
-				UE_LOG(LogTargetPlatformManager, Log, TEXT("Running UBT for AutoSDK init because the file %s is missing"), *(TargetSDKRoot / SDKEnvironmentVarsFile));
+				UE_LOG(LogTargetPlatformManager, Log, TEXT("Running UBT for AutoSDK init because the file %s is missing"), *(TargetSDKRoot / UE::AutoSDK::SDKEnvironmentVarsFile));
 				bNeedsToRun = true;
 			}
 		}
@@ -345,52 +349,53 @@ static bool NeedsToRunSetupPlatformsUBT()
 	return bNeedsToRun;
 }
 
-// kick off a call to UBT nice and early so that it's results are hopefully ready when needed
-FDelayedAutoRegisterHelper GAutoSDKInit(EDelayedRegisterRunPhase::FileSystemReady, []
+void LaunchUBTValidatePlatformsProc(FName PlatformName = NAME_None)
+{
+	SDKSetupUBTProc.Reset();
+
+	FString UBTParams(TEXT("-Mode=ValidatePlatforms -OutputSDKs "));
+	if (PlatformName == NAME_None)
 	{
-		SCOPED_BOOT_TIMING("AutoSDKInit - DelayedAutoRegisterHelper");
+		UBTParams += TEXT(" -AllPlatforms ");
+	}
+	else
+	{
+		UBTParams += FString::Printf(TEXT("  -Platforms=%s "), *(PlatformName.ToString()));
+	}
 
-		// amortize UBT cost by calling it once for all platforms, rather than once per platform.
-		if (IsAutoSDKsEnabled())
-		{
-			FString UBTParams(TEXT("-Mode=ValidatePlatforms -AllPlatforms -OutputSDKs"));
-			{
-				if (FString Project = GetProjectPathForUBT(); Project.Len() > 0)
-				{
-					UBTParams += FString::Printf(TEXT(" -project=%s"), *Project);
-				}
-				// Write the output to a separate file since it performs much better than read pipes in this scenario
-				// where the invocation script is quite involved and is in turn calling other scripts and programs...
-				{
-					FString LogFile = FPaths::Combine(FPaths::ProjectLogDir(), "AutoSDKInfo.txt");
-					FString AbsLogFile = IFileManager::Get().ConvertToAbsolutePathForExternalAppForWrite(*LogFile);
-					UBTParams += FString::Printf(TEXT(" -log=\"%s\" -verbose -timestamps"), *AbsLogFile);
-				}
-			}
-			if (FParse::Param(FCommandLine::Get(), TEXT("Multiprocess")) == false && NeedsToRunSetupPlatformsUBT())
-			{
-				int32 UBTReturnCode = -1;
+	if (FString Project = GetProjectPathForUBT(); Project.Len() > 0)
+	{
+		UBTParams += FString::Printf(TEXT(" -project=%s"), *Project);
+	}
+	// Write the output to a separate file since it performs much better than read pipes in this scenario
+	// where the invocation script is quite involved and is in turn calling other scripts and programs...
+	{
+		FString LogFile = FPaths::Combine(FPaths::ProjectLogDir(), "AutoSDKInfo.txt");
+		FString AbsLogFile = IFileManager::Get().ConvertToAbsolutePathForExternalAppForWrite(*LogFile);
+		UBTParams += FString::Printf(TEXT(" -log=\"%s\" -verbose -timestamps"), *AbsLogFile);
+	}
 
-				void* ReadPipe = nullptr;
-				void* WritePipe = nullptr;
-				AutoSDKSetupUBTProc = FDesktopPlatformModule::Get()->InvokeUnrealBuildToolAsync(UBTParams, *GLog, ReadPipe, WritePipe, true);
-				if (!AutoSDKSetupUBTProc.IsValid())
-				{
-					UE_LOG(LogTargetPlatformManager, Warning, TEXT("AutoSDK is enabled (UE_SDKS_ROOT is set), but failed to run UBT to check SDK status! Check your installation."));
-				}
-			}
-			else
-			{
-				UE_LOGFMT(LogTargetPlatformManager, Log, "Skip running UBT AutoSDK init with params [{UBTParams}]", UBTParams);
-			}
+	if (FParse::Param(FCommandLine::Get(), TEXT("Multiprocess")) == false && NeedsToRunSetupPlatformsUBT())
+	{
+		int32 UBTReturnCode = -1;
+
+		void* ReadPipe = nullptr;
+		void* WritePipe = nullptr;
+		SDKSetupUBTProc = FDesktopPlatformModule::Get()->InvokeUnrealBuildToolAsync(UBTParams, *GLog, ReadPipe, WritePipe, true);
+		if (!SDKSetupUBTProc.IsValid())
+		{
+			UE_LOG(LogTargetPlatformManager, Warning, TEXT("AutoSDK is enabled (UE_SDKS_ROOT is set), but failed to run UBT to check SDK status! Check your installation."));
 		}
 	}
-);
-
 }
-#endif
-
 
+// kick off a call to UBT nice and early so that it's results are hopefully ready when needed
+FDelayedAutoRegisterHelper GSDKInit(EDelayedRegisterRunPhase::FileSystemReady, []
+	{
+		SCOPED_BOOT_TIMING("AutoSDKInit - DelayedAutoRegisterHelper");
+		LaunchUBTValidatePlatformsProc();
+	}
+);
 
 static const size_t MaxPlatformCount = 64;		// In the unlikely event that someone bumps this please note that there's
 												// an implicit assumption that there won't be more than 64 unique target
@@ -472,64 +477,111 @@ ITargetPlatformControls::FOnTargetDeviceLost& ITargetPlatformControls::OnDeviceL
 class FTargetPlatformManagerModule
 	: public ITargetPlatformManagerModule
 {
-public:
-
-	/** Default constructor. */
-	FTargetPlatformManagerModule()
-		: bRestrictFormatsToRuntimeOnly(false)
-		, bForceCacheUpdate(true)
-		, bHasInitErrors(false)
-		, bIgnoreFirstDelegateCall(true)
+	void SDKSetup(FName InPlatformName = NAME_None)
 	{
-#if WITH_EDITOR && UE_WITH_TURNKEY_SUPPORT
+		if (SDKSetupUBTProc.IsValid())
+		{
+			SCOPED_BOOT_TIMING("AutoSDKInit - WaitForUBTProc");
+			// wait for UBT to finish
+			FPlatformProcess::WaitForProc(SDKSetupUBTProc);
 
-		ITurnkeySupportModule::Get().UpdateSdkInfo();
-#endif
+			FString LogFile = FPaths::Combine(FPaths::ProjectLogDir(), "AutoSDKInfo.txt");
+			FString AbsLogFile = IFileManager::Get().ConvertToAbsolutePathForExternalAppForWrite(*LogFile);
 
-#if AUTOSDKS_ENABLED		
-		
-		// AutoSDKs only enabled if UE_SDKS_ROOT is set.
-		if (UE::AutoSDK::IsAutoSDKsEnabled())
-		{					
-			if (UE::AutoSDK::AutoSDKSetupUBTProc.IsValid())
+			TArray<FString> Lines;
+
+			bool bIsAutoSDKEnabled = IsAutoSDKsEnabled();
+			if (FFileHelper::LoadFileToStringArray(Lines, *AbsLogFile) && !bIsAutoSDKEnabled)
 			{
-				SCOPED_BOOT_TIMING("AutoSDKInit - WaitForUBTProc");
-				// wait for UBT to finish
-				FPlatformProcess::WaitForProc(UE::AutoSDK::AutoSDKSetupUBTProc);
-
-				// if it failed, we don't want to cache the saved info, so we make sure we run it again next time
-				int32 ReturnCode = 0;
-				FPlatformProcess::GetProcReturnCode(UE::AutoSDK::AutoSDKSetupUBTProc, &ReturnCode);
-				FPlatformProcess::CloseProc(UE::AutoSDK::AutoSDKSetupUBTProc);
-				UE::AutoSDK::AutoSDKSetupUBTProc.Reset();
-				UE_LOGFMT(LogTargetPlatformManager, Log, "UBT AutoSDK ReturnCode: {ReturnCode}", ReturnCode);
-
-				if (ReturnCode == 0)
+				TMap<FName, FName> UBTNameToIniPlatformName;
+				for (auto Pair : FDataDrivenPlatformInfoRegistry::GetAllPlatformInfos())
 				{
-					// now save out the cached info so next run we can skip UBT
-					UE::AutoSDK::GetCachedInfoToSaveAfterUBT().Write(*(FPaths::ProjectIntermediateDir() / TEXT("CachedAutoSdkInfo.ini")));
+					FName PlatformName = Pair.Key;
+					if (InPlatformName == PlatformName || InPlatformName == NAME_None)
+					{
+						const FDataDrivenPlatformInfo& Info = Pair.Value;
+						UBTNameToIniPlatformName.Add(Info.UBTPlatformName, PlatformName);
+					}
+				}
+
+				for (int Index = 0; Index < Lines.Num(); ++Index)
+				{
+					FString Item = Lines[Index];
+					FString ValidateTag = TEXT("##PlatformValidate: ");
+					int32 StartIndex = Item.Find(ValidateTag, ESearchCase::IgnoreCase);
+					if (StartIndex != INDEX_NONE)
+					{
+						StartIndex += ValidateTag.Len();
+						int32 EndIndex = Item.Find(TEXT(" "), ESearchCase::IgnoreCase, ESearchDir::FromStart, StartIndex);
+						bool bIsSDKValid = !Item.Contains(TEXT("INVALID"));
+						if (EndIndex != INDEX_NONE)
+						{
+							FName UBTName = FName(Item.Mid(StartIndex, EndIndex - StartIndex));
+							FName* IniPlatformName = UBTNameToIniPlatformName.Find(UBTName);
+							if (IniPlatformName != nullptr)
+							{
+								PlatformsSetup.Add(*IniPlatformName, bIsSDKValid);
+							}
+						}
+					}
 				}
 			}
+			// if it failed, we don't want to cache the saved info, so we make sure we run it again next time
+			int32 ReturnCode = 0;
 
-			// we have to setup our local environment according to AutoSDKs or the ITargetPlatform's IsSDkInstalled calls may fail
-			// before we get a change to setup for a given platform.  Use the platforminfo list to avoid any kind of interdependency.
-			SCOPED_BOOT_TIMING("AutoSDKInit - SetupAndValidateAutoSDK");
-			for (auto Pair: FDataDrivenPlatformInfoRegistry::GetAllPlatformInfos())
-			{
-				if (Pair.Value.AutoSDKPath.Len() > 0)
+			FPlatformProcess::GetProcReturnCode(SDKSetupUBTProc, &ReturnCode);
+			FPlatformProcess::CloseProc(SDKSetupUBTProc);
+
+			if (IsAutoSDKsEnabled() && ReturnCode == 0)
 			{
-					SetupAndValidateAutoSDK(Pair.Value.AutoSDKPath);
-				}
+				// now save out the cached info so next run we can skip UBT
+				GetCachedInfoToSaveAfterUBT().Write(*(FPaths::ProjectIntermediateDir() / TEXT("CachedAutoSdkInfo.ini")));
 			}
 
-			FString ManualSDKEnvironmentVarsPath = FPaths::EngineIntermediateDir() / FString(TEXT("ManualSDKEnvVars.txt"));
-			if (IFileManager::Get().FileExists(*ManualSDKEnvironmentVarsPath))
+			UE_LOGFMT(LogTargetPlatformManager, Log, "UBT AutoSDK ReturnCode: {ReturnCode}", ReturnCode);
+		}
+
+		// we have to setup our local environment according to AutoSDKs or the ITargetPlatform's IsSDkInstalled calls may fail
+		// before we get a change to setup for a given platform.  Use the platforminfo list to avoid any kind of interdependency.
+		SCOPED_BOOT_TIMING("AutoSDKInit - SetupAutoSDK");
+		if (IsAutoSDKsEnabled())
+		{
+			for (auto Pair : FDataDrivenPlatformInfoRegistry::GetAllPlatformInfos())
 			{
-				SetupEnvironmentFromManualSDK(ManualSDKEnvironmentVarsPath);
+				FName PlatformName = Pair.Key;
+				if (InPlatformName == PlatformName || InPlatformName == NAME_None)
+				{
+					if (Pair.Value.AutoSDKPath.Len() > 0)
+					{
+						SetupAutoSDK(Pair.Value.AutoSDKPath);
+					}
+				}
 			}
 		}
+
+		FString ManualSDKEnvironmentVarsPath = FPaths::EngineIntermediateDir() / FString(TEXT("ManualSDKEnvVars.txt"));
+		if (IFileManager::Get().FileExists(*ManualSDKEnvironmentVarsPath))
+		{
+			SetupEnvironmentFromManualSDK(ManualSDKEnvironmentVarsPath);
+		}
+	}
+
+public:
+
+	/** Default constructor. */
+	FTargetPlatformManagerModule()
+		: bRestrictFormatsToRuntimeOnly(false)
+		, bForceCacheUpdate(true)
+		, bHasInitErrors(false)
+		, bIgnoreFirstDelegateCall(true)
+	{
+#if WITH_EDITOR && UE_WITH_TURNKEY_SUPPORT
+
+		ITurnkeySupportModule::Get().UpdateSdkInfo();
 #endif
 
+		SDKSetup();
+
 		TextureFormatManager = FModuleManager::LoadModulePtr<ITextureFormatManagerModule>("TextureFormat");
 
 		// Calling a virtual function from a constructor, but with no expectation that a derived implementation of this
@@ -561,9 +613,6 @@ public:
 	virtual void Invalidate() override
 	{
 		bForceCacheUpdate = true;
-
-		SetupSDKStatus();
-		//GetTargetPlatforms(); redudant with next call
 		GetActiveTargetPlatforms();
 
 		bForceCacheUpdate = false;
@@ -931,8 +980,19 @@ public:
 		return nullptr;
 	}
 
-protected:
+	virtual void SetupEnvironmentVariables(TArray<FString>& EnvVarNames, const TArray<FString>& EnvVarValues) override
+	{
+		checkf(EnvVarNames.Num() == EnvVarValues.Num(), TEXT("EnvVarNames and EnvVarValues are expected to have the same number of elements."));
+		for (int i = 0; i < EnvVarNames.Num(); ++i)
+		{
+			const FString& EnvVarName = EnvVarNames[i];
+			const FString& EnvVarValue = EnvVarValues[i];
+			UE_LOG(LogTargetPlatformManager, Verbose, TEXT("Setting variable '%s' to '%s'."), *EnvVarName, *EnvVarValue);
+			FPlatformMisc::SetEnvironmentVar(*EnvVarName, *EnvVarValue);
+		}
+	}
 
+protected:
 
 	bool InitializeSinglePlatform(FName PlatformName, const FString& AutoSDKPath)
 	{
@@ -965,16 +1025,18 @@ protected:
 			ModuleSettings = FModuleManager::LoadModulePtr<ITargetPlatformSettingsModule>(PlatformSettingsModuleName);
 		}
 
-		if (ModuleManager.ModuleExists(*PlatformControlsModuleName.ToString()))
-		{
-			ModuleControls = FModuleManager::LoadModulePtr<ITargetPlatformControlsModule>(PlatformControlsModuleName);
-		}
-
+		bool bIsRunningPlatform = false;
 		if(ModuleSettings != nullptr)
 		{
 			TArray<ITargetPlatformSettings*> TargetPlatformSettings = ModuleSettings->GetTargetPlatformSettings();
 			for (ITargetPlatformSettings* Platform : TargetPlatformSettings)
 			{
+				// one running platform
+				if(Platform->IsRunningPlatform())
+				{
+					checkf(!bIsRunningPlatform, TEXT("Cannot have more than one Running Platform"));
+					bIsRunningPlatform = true;
+				}
 				PlatformSettings.Add(Platform);
 				if (Module != nullptr)
 				{
@@ -991,14 +1053,14 @@ protected:
 				EModuleLoadResult Result;
 				Module = static_cast<ITargetPlatformModule*>(ModuleManager.LoadModuleWithFailureReason(PlatformModuleName, Result, ELoadModuleFlags::LogFailures));
 				UE_CLOGFMT(!Module, LogTargetPlatformManager, Warning,
-					"Failed to load module '{PlatformModuleName}' for platform {PlatformName} (Reason={Reason}, AutoSDKPath='{AutoSDKPath}').",
-					PlatformModuleName, PlatformName.ToString(), LexToString(Result), AutoSDKPath);
+					"Failed to load module '{PlatformModuleName}' for platform {PlatformName} (Reason={Reason}).",
+					PlatformModuleName, PlatformName.ToString(), LexToString(Result));
 			}
 			else
 			{
 				UE_LOGFMT(LogTargetPlatformManager, Log,
-					"Failed to load module '{PlatformModuleName}' for platform {PlatformName} (Reason={Reason}, AutoSDKPath='{AutoSDKPath}').",
-					PlatformModuleName, PlatformName.ToString(), LexToString(EModuleLoadResult::FileNotFound), AutoSDKPath);
+					"Failed to load module '{PlatformModuleName}' for platform {PlatformName} (Reason={Reason}).",
+					PlatformModuleName, PlatformName.ToString(), LexToString(EModuleLoadResult::FileNotFound));
 			}
 		}
 
@@ -1006,12 +1068,19 @@ protected:
 		if (Module)
 		{
 			TRACE_CPUPROFILER_EVENT_SCOPE(FTargetPlatformManagerModule::InitializeSinglePlatform Loading);
-
+			// If AutoSDK is enabled, then we skip the check if AutoSDKPath of this platform is Empty.
+			// If AutoSDK is disabled for whatever reason, then we do not skip and base the platform validity on UBT.
+			bool bSkipSDKCheck = (AutoSDKPath == TEXT("")) && IsAutoSDKsEnabled();
+			FName PlatformSDKNameToCheck = IsAutoSDKsEnabled() ? FName(*AutoSDKPath) : PlatformName;
 			// would like to move this check to GetActiveTargetPlatforms, but too many things cache this result
 			// this setup will become faster after TTP 341897 is complete.
 		RETRY_SETUPANDVALIDATE:
-			if (AutoSDKPath == TEXT("") || SetupAndValidateAutoSDK(AutoSDKPath))
+			if (bSkipSDKCheck || bIsRunningPlatform || ValidateSDK(PlatformSDKNameToCheck))
 			{
+				if (ModuleManager.ModuleExists(*PlatformControlsModuleName.ToString()))
+				{
+					ModuleControls = FModuleManager::LoadModulePtr<ITargetPlatformControlsModule>(PlatformControlsModuleName);
+				}
 
 				if (ModuleControls != nullptr)
 				{
@@ -1047,7 +1116,7 @@ protected:
 						goto RETRY_SETUPANDVALIDATE;
 					}
 				}
-				UE_LOG(LogTargetPlatformManager, Display, TEXT("Failed to SetupAndValidateAutoSDK for platform '%s'"), *PlatformName.ToString());
+				UE_LOG(LogTargetPlatformManager, Display, TEXT("Failed to SetupSDK for platform '%s'"), *PlatformName.ToString());
 			}
 		}
 
@@ -1139,64 +1208,13 @@ protected:
 		}
 	}
 
-	bool UpdatePlatformEnvironment(const FString& PlatformName, TArray<FString> &Keys, TArray<FString> &Values) override
-	{
-		SetupEnvironmentVariables(Keys, Values);
-		return SetupSDKStatus(PlatformName);	
-	}
-
-	bool SetupAndValidateAutoSDK(const FString& AutoSDKPath)
-	{
-		TRACE_CPUPROFILER_EVENT_SCOPE(FTargetPlatformManagerModule::SetupAndValidateAutoSDK);
-
-#if AUTOSDKS_ENABLED
-		bool bValidSDK = false;
-		if (AutoSDKPath.Len() > 0)
-		{		
-			FName PlatformFName(*AutoSDKPath);
-
-			// cache result of the last setup attempt to avoid calling UBT all the time.
-			bool* bPreviousSetupSuccessful = PlatformsSetup.Find(PlatformFName);
-			if (bPreviousSetupSuccessful)
-			{
-				bValidSDK = *bPreviousSetupSuccessful;
-			}
-			else
-			{
-				bValidSDK = SetupEnvironmentFromAutoSDK(AutoSDKPath);
-				PlatformsSetup.Add(PlatformFName, bValidSDK);
-			}
-		}
-		else
-		{
-			// if a platform has no AutoSDKPath, then just assume the SDK is installed, we have no basis for determining it.
-			bValidSDK = true;
-		}
-		return bValidSDK;
-#else
-		return true;
-#endif // AUTOSDKS_ENABLED
-	}
-
-	bool SetupEnvironmentFromManualSDK(const FString& EnvVarFileName)
-	{
-		TRACE_CPUPROFILER_EVENT_SCOPE(FTargetPlatformManagerModule::SetupEnvironmentFromManualSDK);
-
-		UE_LOG(LogTargetPlatformManager, Verbose, TEXT("Reading the manifest for auto-selected manual sdks") );
-		bool bResult = SetupEnvironmentFromEnvVarFile(EnvVarFileName);
-		IFileManager::Get().Delete(*EnvVarFileName);
-		return bResult;
-	}
-	
 	bool SetupEnvironmentFromAutoSDK(const FString& AutoSDKPath)
 	{
 		TRACE_CPUPROFILER_EVENT_SCOPE(FTargetPlatformManagerModule::SetupEnvironmentFromAutoSDK);
 
-#if AUTOSDKS_ENABLED
-		
-		if (!UE::AutoSDK::IsAutoSDKsEnabled())
+		if (!IsAutoSDKsEnabled())
 		{
-			return true;
+			return false;
 		}
 
 		// Invoke UBT to perform SDK switching, or detect that a proper manual SDK is already setup.				
@@ -1206,7 +1224,7 @@ protected:
 		FString HostPlatform = FString::Printf(TEXT("Host%hs"), FPlatformProperties::IniPlatformName());
 #endif		
 
-		FString TargetSDKRoot = UE::AutoSDK::GetAutoSDKPlatformDir(*AutoSDKPath);
+		FString TargetSDKRoot = GetAutoSDKPlatformDir(*AutoSDKPath);
 		FString SDKInstallManifestFilePath = TargetSDKRoot / UE::AutoSDK::SDKInstallManifestFileName;
 
 		TArray<FString> FileLines;
@@ -1223,30 +1241,66 @@ protected:
 			static const FString ManualSDKString(TEXT("ManualSDK"));
 			if (FileLines[1].Compare(ManualSDKString, ESearchCase::IgnoreCase) == 0)
 			{
-				UE_LOG(LogTargetPlatformManager, Verbose, TEXT("Platform %s has manual sdk install"), *AutoSDKPath);				
+				UE_LOG(LogTargetPlatformManager, Verbose, TEXT("Platform %s has manual sdk install"), *AutoSDKPath);
 				return true;
 			}
 		}
 		else
-		{	
-			UE_LOG(LogTargetPlatformManager, Log, TEXT("Install manifest file for Platform %s not found.  Platform not set up."), *AutoSDKPath);			
+		{
+			UE_LOG(LogTargetPlatformManager, Log, TEXT("Install manifest file for Platform %s not found.  Platform not set up."), *AutoSDKPath);
 			return false;			
-		}		
+		}
 
 		static const FString SDKEnvironmentVarsFile(TEXT("OutputEnvVars.txt"));
-		FString EnvVarFileName = FPaths::Combine(*TargetSDKRoot, *SDKEnvironmentVarsFile);		
+		FString EnvVarFileName = FPaths::Combine(*TargetSDKRoot, *SDKEnvironmentVarsFile);
 
 		if (!SetupEnvironmentFromEnvVarFile(EnvVarFileName))
 		{
-			UE_LOG(LogTargetPlatformManager, Warning, TEXT("OutputEnvVars.txt not found for platform: '%s'"), *AutoSDKPath);			
+			UE_LOG(LogTargetPlatformManager, Warning, TEXT("OutputEnvVars.txt not found for platform: '%s'"), *AutoSDKPath);
 			return false;
 		}
 
-		UE_LOG(LogTargetPlatformManager, Verbose, TEXT("Platform %s has auto sdk install"), *AutoSDKPath);		
-		return true;
-#else
+		UE_LOG(LogTargetPlatformManager, Verbose, TEXT("Platform %s has auto sdk install"), *AutoSDKPath);
 		return true;
-#endif
+	}
+
+	void SetupAutoSDK(const FString& AutoSDKPath)
+	{
+		TRACE_CPUPROFILER_EVENT_SCOPE(FTargetPlatformManagerModule::SetupAndValidateAutoSDK);
+		bool bValidSDK = false;
+		if (AutoSDKPath.Len() > 0)
+		{		
+			FName AutoSDKName(*AutoSDKPath);
+
+			// cache result of the last setup attempt to avoid calling UBT all the time.
+			bool* bPreviousSetupSuccessful = PlatformsSetup.Find(AutoSDKName);
+			if (!bPreviousSetupSuccessful)
+			{
+				bValidSDK = SetupEnvironmentFromAutoSDK(AutoSDKPath);
+				PlatformsSetup.Add(AutoSDKName, bValidSDK);
+			}
+		}
+	}
+
+	bool ValidateSDK(FName PlatformName)
+	{
+		bool bValidSDK = false;
+		bool* SetupSuccessful = PlatformsSetup.Find(PlatformName);
+		if (SetupSuccessful)
+		{
+			bValidSDK = *SetupSuccessful;
+		}
+		return bValidSDK;
+	}
+
+	bool SetupEnvironmentFromManualSDK(const FString& EnvVarFileName)
+	{
+		TRACE_CPUPROFILER_EVENT_SCOPE(FTargetPlatformManagerModule::SetupEnvironmentFromManualSDK);
+
+		UE_LOG(LogTargetPlatformManager, Verbose, TEXT("Reading the manifest for auto-selected manual sdks") );
+		bool bResult = SetupEnvironmentFromEnvVarFile(EnvVarFileName);
+		IFileManager::Get().Delete(*EnvVarFileName);
+		return bResult;
 	}
 
 	bool SetupEnvironmentFromEnvVarFile( const FString& EnvVarFileName )
@@ -1369,122 +1423,32 @@ protected:
 		}
 	}
 
-	bool SetupSDKStatus()
-	{
-		return SetupSDKStatus(TEXT(""));
-	}
-
-	bool SetupSDKStatus(const FString& TargetPlatforms)
+	bool UpdateAfterSDKInstall(FName PlatformName)
 	{
-		TRACE_CPUPROFILER_EVENT_SCOPE(FTargetPlatformManagerModule::SetupSDKStatus);
-
-//		FDataDrivenPlatformInfoRegistry::UpdateSdkStatus();
-#if 0
-		DECLARE_SCOPE_CYCLE_COUNTER( TEXT( "FTargetPlatformManagerModule::SetupSDKStatus" ), STAT_FTargetPlatformManagerModule_SetupSDKStatus, STATGROUP_TargetPlatform );
-
-		// run UBT with -validate -allplatforms and read the output
-		FString CmdExe = TEXT("{EngineDir}/Binaries/DotNET/UnrealBuildTool.exe");
-		FString CommandLine = TEXT("-Mode=ValidatePlatforms");
-		// Allow for only a subset of platforms to be reparsed - needed when kicking a change from the UI
-		CommandLine += TargetPlatforms.IsEmpty() ? TEXT(" -allplatforms") : (TEXT(" -platforms=") + TargetPlatforms);
-		
-		// convert into appropriate calls for the current platform
-		FPlatformProcess::ModifyCreateProcParams(CmdExe, CommandLine, FGenericPlatformProcess::ECreateProcHelperFlags::None);
-
+		TRACE_CPUPROFILER_EVENT_SCOPE(FTargetPlatformManagerModule::UpdateAfterSDKInstall);
 
-		TSharedPtr<FMonitoredProcess> UBTProcess = MakeShareable(new FMonitoredProcess(CmdExe, CommandLine, true));
-		UBTProcess->OnOutput().BindStatic(&FTargetPlatformManagerModule::OnStatusOutput);
-		SDKStatusMessage = TEXT("");
+		const FDataDrivenPlatformInfo& Info = FDataDrivenPlatformInfoRegistry::GetPlatformInfo(PlatformName);
 
-		{
-			TRACE_CPUPROFILER_EVENT_SCOPE(FTargetPlatformManagerModule::SetupSDKStatus WaitUntilUBTStarted);
-			UBTProcess->Launch();
-			while(UBTProcess->Update())
-			{
-				FPlatformProcess::Sleep(0.01f);
-			}
-		}
 
-		TArray<FString> PlatArray;
-		SDKStatusMessage.ParseIntoArrayWS(PlatArray);
-		for (int Index = 0; Index < PlatArray.Num()-2; ++Index)
+		FName AutoSDKName = NAME_None;
+		FString AutoSDKPath;
+		if(IsAutoSDKsEnabled())
 		{
-			FString Item = PlatArray[Index];
-			if (PlatArray[Index].Contains(TEXT("##PlatformValidate:")))
+			AutoSDKPath = Info.AutoSDKPath;
+			AutoSDKName = *AutoSDKPath;
+			if (AutoSDKName != NAME_None)
 			{
-				PlatformInfo::EPlatformSDKStatus Status = PlatArray[Index+2].Contains(TEXT("INVALID")) ? PlatformInfo::EPlatformSDKStatus::NotInstalled : PlatformInfo::EPlatformSDKStatus::Installed;
-				FString PlatformName = PlatArray[Index+1];
-				if (PlatformName == TEXT("Win64"))
-				{
-					PlatformName = TEXT("WindowsEditor");
-					PlatformInfo::UpdatePlatformSDKStatus(PlatformName, Status);
-					PlatformName = TEXT("Windows");
-					PlatformInfo::UpdatePlatformSDKStatus(PlatformName, Status);
-					PlatformName = TEXT("WindowsClient");
-					PlatformInfo::UpdatePlatformSDKStatus(PlatformName, Status);
-					PlatformName = TEXT("WindowsServer");
-					PlatformInfo::UpdatePlatformSDKStatus(PlatformName, Status);
-				}
-				else if (PlatformName == TEXT("Mac"))
-				{
-					PlatformInfo::UpdatePlatformSDKStatus(PlatformName, Status);
-					PlatformName = TEXT("Mac");
-					PlatformInfo::UpdatePlatformSDKStatus(PlatformName, Status);
-					PlatformName = TEXT("MacClient");
-					PlatformInfo::UpdatePlatformSDKStatus(PlatformName, Status);
-					PlatformName = TEXT("MacServer");
-					PlatformInfo::UpdatePlatformSDKStatus(PlatformName, Status);
-				}
-				else if (PlatformName == TEXT("Linux"))
-				{
-					PlatformInfo::UpdatePlatformSDKStatus(PlatformName, Status);
-					PlatformName = TEXT("Linux");
-					PlatformInfo::UpdatePlatformSDKStatus(PlatformName, Status);
-					PlatformName = TEXT("LinuxClient");
-					PlatformInfo::UpdatePlatformSDKStatus(PlatformName, Status);
-					PlatformName = TEXT("LinuxServer");
-					PlatformInfo::UpdatePlatformSDKStatus(PlatformName, Status);
-				}
-				else if (PlatformName == TEXT("LinuxArm64"))
-				{
-					PlatformInfo::UpdatePlatformSDKStatus(PlatformName, Status);
-					PlatformName = TEXT("LinuxArm64");
-					PlatformInfo::UpdatePlatformSDKStatus(PlatformName, Status);
-					PlatformName = TEXT("LinuxArm64Client");
-					PlatformInfo::UpdatePlatformSDKStatus(PlatformName, Status);
-					PlatformName = TEXT("LinuxArm64Server");
-					PlatformInfo::UpdatePlatformSDKStatus(PlatformName, Status);
-				}
-				else if (PlatformName == TEXT("Desktop"))
-				{
-					// since Desktop is just packaging, we don't need an SDK, and UBT will return INVALID, since it doesn't build for it
-					PlatformInfo::UpdatePlatformSDKStatus(PlatformName, PlatformInfo::EPlatformSDKStatus::Installed);
-				}
-				else
-				{
-					PlatformInfo::UpdatePlatformSDKStatus(PlatformName, Status);
-				}
+				// make sure we can re-do the AutoSDK setup
+				PlatformsSetup.Remove(AutoSDKName);
 			}
 		}
-#endif
-		return true;
-	}
-
-	bool UpdateAfterSDKInstall(FName PlatformName)
-	{
-		TRACE_CPUPROFILER_EVENT_SCOPE(FTargetPlatformManagerModule::UpdateAfterSDKInstall);
-
-		const FDataDrivenPlatformInfo& Info = FDataDrivenPlatformInfoRegistry::GetPlatformInfo(PlatformName);
-
-#if AUTOSDKS_ENABLED
-		FString AutoSDKPath = Info.AutoSDKPath;
-		FName AutoSDKName(*AutoSDKPath);
-		if (AutoSDKName != NAME_None)
+		else
 		{
-			// make sure we can re-do the AutoSDK setup
-			PlatformsSetup.Remove(AutoSDKName);
+			// Must call UBT again to update platform status if AutoSDK is off.
+			LaunchUBTValidatePlatformsProc(PlatformName);
 		}
-#endif
+		
+		SDKSetup(PlatformName);
 
 		// note: this assumes, along with other Turnkey code, that there is a TargetPlatform named with the IniPlatformName
 		ITargetPlatform* TargetPlatform = FindTargetPlatform(PlatformName);
@@ -1501,14 +1465,14 @@ protected:
 		}
 		else
 		{
-#if AUTOSDKS_ENABLED
-			if (AutoSDKName != NAME_None)
+			if (IsAutoSDKsEnabled())
 			{
-				// setup AutoSDK, and then re-initialize the TP
-				SetupAndValidateAutoSDK(AutoSDKPath);
+				if (!AutoSDKPath.IsEmpty())
+				{
+					// setup AutoSDK, and then re-initialize the TP
+					SetupAutoSDK(AutoSDKPath);
+				}
 			}
-#endif
-
 			bTPInitialized = TargetPlatform->InitializeHostPlatform();
 		}
 
@@ -1527,17 +1491,6 @@ protected:
 
 private:
 
-	void SetupEnvironmentVariables(TArray<FString> &EnvVarNames, const TArray<FString>& EnvVarValues)
-	{
-		for (int i = 0; i < EnvVarNames.Num(); ++i)
-		{
-			const FString& EnvVarName = EnvVarNames[i];
-			const FString& EnvVarValue = EnvVarValues[i];
-			UE_LOG(LogTargetPlatformManager, Verbose, TEXT("Setting variable '%s' to '%s'."), *EnvVarName, *EnvVarValue);
-			FPlatformMisc::SetEnvironmentVar(*EnvVarName, *EnvVarValue);
-		}
-	}
-
 	void ModulesChangesCallback(FName ModuleName, EModuleChangeReason ReasonForChange)
 	{
 		if (!bIgnoreFirstDelegateCall && ModuleName.ToString().Contains(TEXT("TargetPlatform")) && !ModuleName.ToString().Contains(TEXT("ProjectTargetPlatformEditor")))
@@ -1780,10 +1733,9 @@ private:
 	// External module that texture format operations are forwarded to
 	ITextureFormatManagerModule* TextureFormatManager;
 
-#if AUTOSDKS_ENABLED
 	// holds the list of Platforms that have attempted setup.
+	// When AutoSDK is enabled this holds the AutoSDKName because multiple platform can use the same AutoSDKName
 	TMap<FName, bool> PlatformsSetup;
-#endif
 };
 
 
diff --git a/Engine/Source/Developer/TargetPlatform/Public/Common/TargetPlatformBase.h b/Engine/Source/Developer/TargetPlatform/Public/Common/TargetPlatformBase.h
--- a/Engine/Source/Developer/TargetPlatform/Public/Common/TargetPlatformBase.h
+++ b/Engine/Source/Developer/TargetPlatform/Public/Common/TargetPlatformBase.h
@@ -620,7 +620,6 @@ public:
 	virtual bool IsRunningPlatform() const override
 	{
 		// IsRunningPlatform is only for editor platforms
-
 		return false;
 	}
 
@@ -739,7 +738,7 @@ public:
 	}
 	virtual bool IsRunningPlatform() const override 
 	{
-		return TargetPlatformControls->IsRunningPlatform();
+		return TargetPlatformSettings->IsRunningPlatform();
 	}
 	virtual bool IsServerOnly() const override 
 	{
diff --git a/Engine/Source/Developer/TargetPlatform/Public/Common/TargetPlatformControlsBase.h b/Engine/Source/Developer/TargetPlatform/Public/Common/TargetPlatformControlsBase.h
--- a/Engine/Source/Developer/TargetPlatform/Public/Common/TargetPlatformControlsBase.h
+++ b/Engine/Source/Developer/TargetPlatform/Public/Common/TargetPlatformControlsBase.h
@@ -420,13 +420,6 @@ public:
 		return bIsClientOnly;
 	}
 
-	virtual bool IsRunningPlatform() const override
-	{
-		// IsRunningPlatform is only for editor platforms
-
-		return false;
-	}
-
 protected:
 	// true if this target platform is client-only, ie strips out server stuff
 	bool bIsClientOnly;
diff --git a/Engine/Source/Developer/TargetPlatform/Public/Common/TargetPlatformSettingsBase.h b/Engine/Source/Developer/TargetPlatform/Public/Common/TargetPlatformSettingsBase.h
--- a/Engine/Source/Developer/TargetPlatform/Public/Common/TargetPlatformSettingsBase.h
+++ b/Engine/Source/Developer/TargetPlatform/Public/Common/TargetPlatformSettingsBase.h
@@ -67,6 +67,12 @@ public:
 		return false;
 	}
 
+	virtual bool IsRunningPlatform() const override
+	{
+		// IsRunningPlatform is only for editor platforms
+		return false;
+	}
+
 protected:
 
 	TARGETPLATFORM_API ERayTracingRuntimeMode ParseRuntimeRayTracingMode(const TCHAR* SectionName) const;
diff --git a/Engine/Source/Developer/TargetPlatform/Public/Interfaces/ITargetPlatformControls.h b/Engine/Source/Developer/TargetPlatform/Public/Interfaces/ITargetPlatformControls.h
--- a/Engine/Source/Developer/TargetPlatform/Public/Interfaces/ITargetPlatformControls.h
+++ b/Engine/Source/Developer/TargetPlatform/Public/Interfaces/ITargetPlatformControls.h
@@ -227,16 +227,6 @@ public:
 	 */
 	virtual bool IsLittleEndian() const = 0;
 
-	/**
-	 * Checks whether this platform is the platform that's currently running.
-	 *
-	 * For example, when running on Windows, the Windows ITargetPlatform will return true
-	 * and all other platforms will return false.
-	 *
-	 * @return true if this platform is running, false otherwise.
-	 */
-	virtual bool IsRunningPlatform() const = 0;
-
 	/**
 	 * Checks whether this platform is only a server.
 	 *
diff --git a/Engine/Source/Developer/TargetPlatform/Public/Interfaces/ITargetPlatformManagerModule.h b/Engine/Source/Developer/TargetPlatform/Public/Interfaces/ITargetPlatformManagerModule.h
--- a/Engine/Source/Developer/TargetPlatform/Public/Interfaces/ITargetPlatformManagerModule.h
+++ b/Engine/Source/Developer/TargetPlatform/Public/Interfaces/ITargetPlatformManagerModule.h
@@ -189,9 +189,9 @@ public:
 	virtual uint32 ShaderFormatVersion(FName Name) = 0;
 
 	/**
-	 * Allows changes to environment for a given platform
+	 * Allows changes to the environment
 	 */
-	virtual bool UpdatePlatformEnvironment(const FString& PlatformName, TArray<FString> &Keys, TArray<FString> &Values) = 0;
+	virtual void SetupEnvironmentVariables(TArray<FString>& EnvVarNames, const TArray<FString>& EnvVarValues) = 0;
 
 	/** A callback that holders of ITargetPlatform* must subscribe to to be notified of when the ITargetPlatform* has been invalidated and should be requeried from e.g. FindTargetPlatform */
 	virtual FOnTargetPlatformsInvalidated& GetOnTargetPlatformsInvalidatedDelegate() = 0;
diff --git a/Engine/Source/Developer/TargetPlatform/Public/Interfaces/ITargetPlatformSettings.h b/Engine/Source/Developer/TargetPlatform/Public/Interfaces/ITargetPlatformSettings.h
--- a/Engine/Source/Developer/TargetPlatform/Public/Interfaces/ITargetPlatformSettings.h
+++ b/Engine/Source/Developer/TargetPlatform/Public/Interfaces/ITargetPlatformSettings.h
@@ -318,6 +318,17 @@ public:
 	 */
 	virtual const class FStaticMeshLODSettings& GetStaticMeshLODSettings() const = 0;
 #endif
+
+	/**
+	 * Checks whether this platform is the platform that's currently running.
+	 *
+	 * For example, when running on Windows, the Windows ITargetPlatform will return true
+	 * and all other platforms will return false.
+	 *
+	 * @return true if this platform is running, false otherwise.
+	 */
+	virtual bool IsRunningPlatform() const = 0;
+
 	/** Virtual destructor. */
 	virtual ~ITargetPlatformSettings() = default;
 };
diff --git a/Engine/Source/Developer/Windows/WindowsTargetPlatfomControls/Public/GenericWindowsTargetPlatformControls.h b/Engine/Source/Developer/Windows/WindowsTargetPlatfomControls/Public/GenericWindowsTargetPlatformControls.h
--- a/Engine/Source/Developer/Windows/WindowsTargetPlatfomControls/Public/GenericWindowsTargetPlatformControls.h
+++ b/Engine/Source/Developer/Windows/WindowsTargetPlatfomControls/Public/GenericWindowsTargetPlatformControls.h
@@ -111,12 +111,6 @@ public:
 		return nullptr;
 	}
 
-	virtual bool IsRunningPlatform( ) const override
-	{
-		// Must be Windows platform as editor for this to be considered a running platform
-		return PLATFORM_WINDOWS && !UE_SERVER && !UE_GAME && WITH_EDITOR && TProperties::HasEditorOnlyData();
-	}
-
 	virtual void GetShaderCompilerDependencies(TArray<FString>& OutDependencies) const override
 	{
 #if  PLATFORM_CPU_ARM_FAMILY && !PLATFORM_WINDOWS_ARM64EC
diff --git a/Engine/Source/Developer/Windows/WindowsTargetPlatformSettings/Public/GenericWindowsTargetPlatformSettings.h b/Engine/Source/Developer/Windows/WindowsTargetPlatformSettings/Public/GenericWindowsTargetPlatformSettings.h
--- a/Engine/Source/Developer/Windows/WindowsTargetPlatformSettings/Public/GenericWindowsTargetPlatformSettings.h
+++ b/Engine/Source/Developer/Windows/WindowsTargetPlatformSettings/Public/GenericWindowsTargetPlatformSettings.h
@@ -107,6 +107,12 @@ public:
 		return TSuper::SupportsFeature(Feature);
 	}
 
+	virtual bool IsRunningPlatform() const override
+	{
+		// Must be Windows platform as editor for this to be considered a running platform
+		return PLATFORM_WINDOWS && !UE_SERVER && !UE_GAME && WITH_EDITOR && TProperties::HasEditorOnlyData();
+	}
+
 	virtual void GetAllPossibleShaderFormats( TArray<FName>& OutFormats ) const override
 	{
 		// no shaders needed for dedicated server target
