GetTextureFormatModuleHints is not used and can be deleted

#Synced-CL 49850269

diff --git a/Engine/Source/Developer/TargetPlatform/Public/Common/TargetPlatformBase.h b/Engine/Source/Developer/TargetPlatform/Public/Common/TargetPlatformBase.h
--- a/Engine/Source/Developer/TargetPlatform/Public/Common/TargetPlatformBase.h
+++ b/Engine/Source/Developer/TargetPlatform/Public/Common/TargetPlatformBase.h
@@ -250,21 +250,8 @@ public:
 	{
 	}
 	
-	// beware this is duplicated between TargetPlatformBase and TargetPlatformControlsBase
 	virtual void GetTextureFormatModuleHints(TArray<FName>& OutModuleNames) const override
 	{
-		// these are the default texture format modules, since many platforms 
-		OutModuleNames.Add(TEXT("TextureFormatUncompressed"));
-		OutModuleNames.Add(TEXT("TextureFormatDXT"));
-		OutModuleNames.Add(TEXT("TextureFormatIntelISPCTexComp"));
-		OutModuleNames.Add(TEXT("TextureFormatOodle"));
-
-		// there is a possible optional format module name in the ini (alternate texture compressor)
-		FString TextureCompressionFormat;
-		if (GetConfigSystem()->GetString(TEXT("AlternateTextureCompression"), TEXT("TextureCompressionFormat"), TextureCompressionFormat, GEngineIni))
-		{
-			OutModuleNames.Add(*TextureCompressionFormat);
-		}
 	}
 	
 	// beware this is duplicated between TargetPlatformBase and TargetPlatformControlsBase
diff --git a/Engine/Source/Developer/TargetPlatform/Public/Common/TargetPlatformControlsBase.h b/Engine/Source/Developer/TargetPlatform/Public/Common/TargetPlatformControlsBase.h
--- a/Engine/Source/Developer/TargetPlatform/Public/Common/TargetPlatformControlsBase.h
+++ b/Engine/Source/Developer/TargetPlatform/Public/Common/TargetPlatformControlsBase.h
@@ -161,18 +161,6 @@ public:
 	// beware this is duplicated between TargetPlatformBase and TargetPlatformControlsBase
 	virtual void GetTextureFormatModuleHints(TArray<FName>& OutModuleNames) const override
 	{
-		// these are the default texture format modules, since many platforms 
-		OutModuleNames.Add(TEXT("TextureFormatUncompressed"));
-		OutModuleNames.Add(TEXT("TextureFormatDXT"));
-		OutModuleNames.Add(TEXT("TextureFormatIntelISPCTexComp"));
-		OutModuleNames.Add(TEXT("TextureFormatOodle"));
-
-		// there is a possible optional format module name in the ini (alternate texture compressor)
-		FString TextureCompressionFormat;
-		if (TargetPlatformSettings->GetConfigSystem()->GetString(TEXT("AlternateTextureCompression"), TEXT("TextureCompressionFormat"), TextureCompressionFormat, GEngineIni))
-		{
-			OutModuleNames.Add(*TextureCompressionFormat);
-		}
 	}
 
 	TARGETPLATFORM_API virtual FName GetWaveFormat(const class USoundWave* Wave) const override;
diff --git a/Engine/Source/Developer/TargetPlatform/Public/Interfaces/ITargetPlatformControls.h b/Engine/Source/Developer/TargetPlatform/Public/Interfaces/ITargetPlatformControls.h
--- a/Engine/Source/Developer/TargetPlatform/Public/Interfaces/ITargetPlatformControls.h
+++ b/Engine/Source/Developer/TargetPlatform/Public/Interfaces/ITargetPlatformControls.h
@@ -364,9 +364,7 @@ public:
 	virtual void GetAllTextureFormats(TArray<FName>& OutFormats) const = 0;
 
 	/**
-	 * Gets a list of modules that may contain the GetAllTextureFormats. This is optional -
-	 * if any required texture format isn't found in this list, then it will use the old path
-	 * of loading all texture format modules to gather all available texture formats
+	 * GetTextureFormatModuleHints is not used and should be deleted
 	 */
 	virtual void GetTextureFormatModuleHints(TArray<FName>& OutModuleNames) const = 0;
 
