For Programs always disable UBT checking if AutoSDK is disabled.

#Synced-CL 50069373

diff --git a/Engine/Source/Developer/TargetPlatform/Private/TargetPlatformManagerModule.cpp b/Engine/Source/Developer/TargetPlatform/Private/TargetPlatformManagerModule.cpp
--- a/Engine/Source/Developer/TargetPlatform/Private/TargetPlatformManagerModule.cpp
+++ b/Engine/Source/Developer/TargetPlatform/Private/TargetPlatformManagerModule.cpp
@@ -228,7 +228,12 @@ static bool NeedsToRunSetupPlatformsUBT()
 
 	if (!IsAutoSDKsEnabled())
 	{
+		// For Programs always return false if AutoSDK is disabled.
+#if IS_PROGRAM
+		return false;
+#else
 		return true;
+#endif
 	}
 
 	// early out and always return true on builders
