Fixed platform detection that would incorrectly detected a platform as not available in the Editor when AutoSDK is disabled and the pathname to the .uproject contains whitespaces.

#Synced-CL 50163525

diff --git a/Engine/Source/Developer/TargetPlatform/Private/TargetPlatformManagerModule.cpp b/Engine/Source/Developer/TargetPlatform/Private/TargetPlatformManagerModule.cpp
--- a/Engine/Source/Developer/TargetPlatform/Private/TargetPlatformManagerModule.cpp
+++ b/Engine/Source/Developer/TargetPlatform/Private/TargetPlatformManagerModule.cpp
@@ -375,7 +375,7 @@ void LaunchUBTValidatePlatformsProc(FName PlatformName = NAME_None)
 
 	if (FString Project = GetProjectPathForUBT(); Project.Len() > 0)
 	{
-		UBTParams += FString::Printf(TEXT(" -project=%s"), *Project);
+		UBTParams += FString::Printf(TEXT(" -project=\"%s\""), *Project);
 	}
 	// Write the output to a separate file since it performs much better than read pipes in this scenario
 	// where the invocation script is quite involved and is in turn calling other scripts and programs...
