Clarify comment to explain the reason we skip UBT on program target types

#Synced-CL 50069373

diff --git a/Engine/Source/Developer/TargetPlatform/Private/TargetPlatformManagerModule.cpp b/Engine/Source/Developer/TargetPlatform/Private/TargetPlatformManagerModule.cpp
--- a/Engine/Source/Developer/TargetPlatform/Private/TargetPlatformManagerModule.cpp
+++ b/Engine/Source/Developer/TargetPlatform/Private/TargetPlatformManagerModule.cpp
@@ -228,7 +228,9 @@ static bool NeedsToRunSetupPlatformsUBT()
 
 	if (!IsAutoSDKsEnabled())
 	{
-		// For Programs always return false if AutoSDK is disabled.
+		// UBT is usually always invoked when AutoSDK is disabled to check platform status, as there is no other way to check sdk validity. 
+		// This causes a startup overhead and unexpected console log output because we have to launch an UBT process and read the output in order to check platform validity.
+		// Returning false here for program target type solves the problem by skipping calling UBT, we assume all platforms are valid.
 #if IS_PROGRAM
 		return false;
 #else
