Allow to only load host and target platform modules when cooking.

#rnx

#Synced-CL 51150368

diff --git a/Engine/Source/Developer/TargetPlatform/Private/TargetPlatformManagerModule.cpp b/Engine/Source/Developer/TargetPlatform/Private/TargetPlatformManagerModule.cpp
--- a/Engine/Source/Developer/TargetPlatform/Private/TargetPlatformManagerModule.cpp
+++ b/Engine/Source/Developer/TargetPlatform/Private/TargetPlatformManagerModule.cpp
@@ -1179,8 +1179,10 @@ protected:
 		}
 #endif
 
-	// Get the platform we are previewing if GMaxRHIShaderPlatform is a preview SP
+		TSet<FName> AlreadyInitialized;
+		
 #if WITH_EDITOR
+		// Get the platform we are previewing if GMaxRHIShaderPlatform is a preview SP
 		FName PlatformNamePreview = NAME_None;
 		if (IsRunningGame())
 		{
@@ -1189,6 +1191,86 @@ protected:
 				PlatformNamePreview = FDataDrivenShaderPlatformInfo::GetPlatformName(GMaxRHIShaderPlatform);
 			}
 		}
+		
+		if (IsRunningCookCommandlet() && FParse::Param(FCommandLine::Get(), TEXT("-LoadOnlyHostAndTargetPlatformModules")))
+		{
+			// For every given target platform module, attempt to find the containing plugin (e.g. IOSClient -> IOS).
+			//
+			// For successful loading:
+			// * The containing platform must exist.
+			// * The containing platform must be enabled for use.
+			// * The loaded modules after init must contain the requested platform.
+			//
+			// Any failure will fall back to loading all discovered modules.
+				
+			bool bSuccess = false;
+			const TMap<FName, FDataDrivenPlatformInfo>& AllPlatformInfos = FDataDrivenPlatformInfoRegistry::GetAllPlatformInfos();
+			
+			FString TargetPlatformCmdLine;
+			if (FParse::Value(FCommandLine::Get(), TEXT("TargetPlatform="), TargetPlatformCmdLine) && TargetPlatformCmdLine != TEXT("All"))
+			{
+				bSuccess = true;
+				TArray<FString> TargetPlatformNames;
+				TargetPlatformCmdLine.ParseIntoArray(TargetPlatformNames, TEXT("+"), true);
+				
+				for (const FString& TargetPlatformName: TargetPlatformNames)
+				{
+					FString Stripped = TargetPlatformName;
+					Stripped.RemoveFromEnd("Client");
+					Stripped.RemoveFromEnd("Server");
+					Stripped.RemoveFromEnd("Editor");
+					// Cook flavours, for e.g. Android, are defined as Android_ASTC etc.
+					Stripped.LeftInline(Stripped.Contains(TEXT("_")) ? Stripped.Find(TEXT("_")) : Stripped.Len());	
+					
+					FName PlatformToLoad(Stripped);
+					
+					if (!AllPlatformInfos.Contains(PlatformToLoad))
+					{
+						UE_LOG(LogTargetPlatformManager, Warning, TEXT("Assumed parent platform (%s) not found for %s"), *PlatformToLoad.ToString(), *TargetPlatformName);
+						bSuccess = false;
+						break;
+					}
+					
+					const FDataDrivenPlatformInfo& Info = AllPlatformInfos[PlatformToLoad];
+					
+					if (!Info.bEnabledForUse)
+					{
+						UE_LOG(LogTargetPlatformManager, Warning, TEXT("Assumed parent platform (%s) was not enabled for use"), *PlatformToLoad.ToString());
+						bSuccess = false;
+						break;
+					}
+					
+					if (InitializeSinglePlatform(PlatformToLoad, Info.AutoSDKPath))
+					{
+						AlreadyInitialized.Add(PlatformToLoad);
+						if (!PlatformsByName.Contains(FName(TargetPlatformName)))
+						{
+							UE_LOG(LogTargetPlatformManager, Warning, TEXT("Target platform %s does not appear to be part of %s"), *TargetPlatformName, *PlatformToLoad.ToString());
+							bSuccess = false;
+							break;
+						}
+					}
+					else
+					{
+						UE_LOG(LogTargetPlatformManager, Warning, TEXT("Assumed parent platform %s failed to initialize"), *PlatformToLoad.ToString());
+						bSuccess = false;
+					}
+				}
+			}
+			
+			if (bSuccess)
+			{
+				// Host platform is always required so ensure it is initialized if it wasn't the requested.
+				const FName HostPlatformName = FPlatformProperties::IniPlatformName();
+				if (!AlreadyInitialized.Contains(HostPlatformName))
+				{
+					InitializeSinglePlatform(HostPlatformName, AllPlatformInfos[HostPlatformName].AutoSDKPath);
+				}
+				
+				return;
+			}
+		}
+		
 #endif
 
 		// find a set of valid target platform names (the platform DataDrivenPlatformInfo.ini file was found indicates support for the platform 
@@ -1220,7 +1302,7 @@ protected:
 			}
 #endif
 
-			if (Info.bEnabledForUse)
+			if (Info.bEnabledForUse && !AlreadyInitialized.Contains(PlatformName))
 			{
 				InitializeSinglePlatform(PlatformName, Info.AutoSDKPath);
 			}
