Custom Platform Modules don't participate in the SDK management system, so we should skip SDK validation

#Synced-CL 51760584

diff --git a/Engine/Source/Developer/TargetPlatform/Private/TargetPlatformManagerModule.cpp b/Engine/Source/Developer/TargetPlatform/Private/TargetPlatformManagerModule.cpp
--- a/Engine/Source/Developer/TargetPlatform/Private/TargetPlatformManagerModule.cpp
+++ b/Engine/Source/Developer/TargetPlatform/Private/TargetPlatformManagerModule.cpp
@@ -1091,7 +1091,13 @@ protected:
 			TRACE_CPUPROFILER_EVENT_SCOPE(FTargetPlatformManagerModule::InitializeSinglePlatform Loading);
 			// If AutoSDK is enabled, then we skip the check if AutoSDKPath of this platform is Empty.
 			// If AutoSDK is disabled for whatever reason, then we do not skip and base the platform validity on UBT.
-			bool bSkipSDKCheck = false;
+			// If the platform is not in the DataDrivenPlatformInfoRegistry, it's a custom platform module
+			// that doesn't participate in the standard SDK management system, so skip SDK validation
+			bool bSkipSDKCheck = !FDataDrivenPlatformInfoRegistry::GetAllPlatformInfos().Contains(PlatformName);
+			if (bSkipSDKCheck)
+			{
+				UE_LOGF(LogTargetPlatformManager, Log, "Platform '%ls' not in DataDrivenPlatformInfoRegistry, skipping SDK validation", *PlatformName.ToString());
+			}
 			FName PlatformSDKNameToCheck = PlatformName;
 			// would like to move this check to GetActiveTargetPlatforms, but too many things cache this result
 			// this setup will become faster after TTP 341897 is complete.
@@ -1099,7 +1105,7 @@ protected:
 			// This is done here in case autosdk setup failed because shared files are in use by another child cooker
 			if (IsAutoSDKsEnabled())
 			{
-				bSkipSDKCheck = (AutoSDKPath == TEXT(""));
+				bSkipSDKCheck = bSkipSDKCheck || (AutoSDKPath == TEXT(""));
 				PlatformSDKNameToCheck = FName(*AutoSDKPath);
 				if (!bSkipSDKCheck)
 				{
