Fix to recursive call to FTargetPlatformManagerModule::ModulesChangesCallback during post Turnkey installation that causes a dead lock

#Synced-CL 52512740

diff --git a/Engine/Source/Developer/TargetPlatform/Private/TargetPlatformManagerModule.cpp b/Engine/Source/Developer/TargetPlatform/Private/TargetPlatformManagerModule.cpp
--- a/Engine/Source/Developer/TargetPlatform/Private/TargetPlatformManagerModule.cpp
+++ b/Engine/Source/Developer/TargetPlatform/Private/TargetPlatformManagerModule.cpp
@@ -1544,9 +1544,16 @@ private:
 
 	void ModulesChangesCallback(FName ModuleName, EModuleChangeReason ReasonForChange)
 	{
+		if (bIsInDelegateCall)
+		{
+			return;
+		}
+		
 		if (!bIgnoreFirstDelegateCall && ModuleName.ToString().Contains(TEXT("TargetPlatform")) && !ModuleName.ToString().Contains(TEXT("ProjectTargetPlatformEditor")))
 		{
+			bIsInDelegateCall = true;
 			Invalidate();
+			bIsInDelegateCall = false;
 		}
 		bIgnoreFirstDelegateCall = false;
 	}
@@ -1779,9 +1786,12 @@ private:
 	// Flag to indicate that there were errors during initialization
 	bool bHasInitErrors;
 
-	// Flag to avoid redunant reloads
+	// Flag to avoid redundant reloads
 	bool bIgnoreFirstDelegateCall;
 	
+	// Flag to avoid recursive reloads
+	bool bIsInDelegateCall = false;
+	
 	// Holds the list of discovered platforms.
 	TArray<ITargetPlatform*> Platforms;
 	TArray<ITargetPlatformControls*> PlatformControls;
